Fix bloom drawing of pose link when weight is greater than 1

diff --git a/Engine/Source/Editor/AnimationBlueprintEditor/Private/AnimGraphConnectionDrawingPolicy.cpp b/Engine/Source/Editor/AnimationBlueprintEditor/Private/AnimGraphConnectionDrawingPolicy.cpp
--- a/Engine/Source/Editor/AnimationBlueprintEditor/Private/AnimGraphConnectionDrawingPolicy.cpp
+++ b/Engine/Source/Editor/AnimationBlueprintEditor/Private/AnimGraphConnectionDrawingPolicy.cpp
@@ -66,7 +66,7 @@ void FAnimGraphConnectionDrawingPolicy::BuildExecutionRoadmap()
 						FExecPairingMap& Predecessors = PredecessorPins.FindOrAdd((UEdGraphNode*)SourceNode);
 						FTimePair& Timings = Predecessors.FindOrAdd(PoseNet);
 						Timings.PredExecTime = 0.0;
-						Timings.ThisExecTime = VisitRecord.Weight;
+						Timings.ThisExecTime = FMath::Clamp(VisitRecord.Weight, 0.f, 1.f);
 					}
 				}
 			}
