Stop blueprint compiler ensuring when we are currently previewing live link

diff --git a/Engine/Source/Editor/AnimationBlueprintEditor/Private/AnimationBlueprintEditor.cpp b/Engine/Source/Editor/AnimationBlueprintEditor/Private/AnimationBlueprintEditor.cpp
--- a/Engine/Source/Editor/AnimationBlueprintEditor/Private/AnimationBlueprintEditor.cpp
+++ b/Engine/Source/Editor/AnimationBlueprintEditor/Private/AnimationBlueprintEditor.cpp
@@ -1037,6 +1037,7 @@ void FAnimationBlueprintEditor::Compile()
 {
 	// Grab the currently debugged object, so we can re-set it below
 	USkeletalMeshComponent* DebuggedMeshComponent = nullptr;
+	
 	if (UBlueprint* Blueprint = GetBlueprintObj())
 	{
 		UAnimInstance* CurrentDebugObject = Cast<UAnimInstance>(Blueprint->GetObjectBeingDebugged());
@@ -1059,7 +1060,13 @@ void FAnimationBlueprintEditor::Compile()
 			DebuggedMeshComponent->InitAnim(true);
 		}
 
-		GetBlueprintObj()->SetObjectBeingDebugged(DebuggedMeshComponent->GetAnimInstance());
+		UAnimInstance* NewInstance = DebuggedMeshComponent->GetAnimInstance();
+		UBlueprint* Blueprint = GetBlueprintObj();
+
+		if(NewInstance->IsA(Blueprint->GeneratedClass))
+		{
+			Blueprint->SetObjectBeingDebugged(NewInstance);
+		}
 	}
 
 	// reset the selected skeletal control node
