Fixed pointer use before verify static analysis warning
#rnx

diff --git a/Engine/Source/Editor/AnimationBlueprintEditor/Private/AnimationBlueprintEditor.cpp b/Engine/Source/Editor/AnimationBlueprintEditor/Private/AnimationBlueprintEditor.cpp
--- a/Engine/Source/Editor/AnimationBlueprintEditor/Private/AnimationBlueprintEditor.cpp
+++ b/Engine/Source/Editor/AnimationBlueprintEditor/Private/AnimationBlueprintEditor.cpp
@@ -1164,7 +1164,7 @@ void FAnimationBlueprintEditor::Compile()
 
 		// re-apply preview anim bp if needed
 		UAnimBlueprint* AnimBlueprint = GetAnimBlueprint();
-		UAnimBlueprint* PreviewAnimBlueprint = AnimBlueprint->GetPreviewAnimationBlueprint();
+		UAnimBlueprint* PreviewAnimBlueprint = AnimBlueprint ? AnimBlueprint->GetPreviewAnimationBlueprint() : nullptr;
 		
 		if (PreviewAnimBlueprint)
 		{
@@ -1172,7 +1172,7 @@ void FAnimationBlueprintEditor::Compile()
 		}
 
 		UAnimInstance* NewInstance = DebuggedMeshComponent->GetAnimInstance();
-		if(AnimBlueprint && (NewInstance->IsA(AnimBlueprint->GeneratedClass) || (PreviewAnimBlueprint != nullptr && NewInstance->IsA(PreviewAnimBlueprint->GeneratedClass))))
+		if((AnimBlueprint && (NewInstance->IsA(AnimBlueprint->GeneratedClass)) || (PreviewAnimBlueprint && NewInstance->IsA(PreviewAnimBlueprint->GeneratedClass))))
 		{
 			PersonaUtils::SetObjectBeingDebugged(AnimBlueprint, NewInstance);
 		}
