Fix compile error: '&&' within '||' [-Werror,-Wlogical-op-parentheses]

diff --git a/Engine/Source/Editor/AnimationBlueprintEditor/Private/AnimationBlueprintEditor.cpp b/Engine/Source/Editor/AnimationBlueprintEditor/Private/AnimationBlueprintEditor.cpp
--- a/Engine/Source/Editor/AnimationBlueprintEditor/Private/AnimationBlueprintEditor.cpp
+++ b/Engine/Source/Editor/AnimationBlueprintEditor/Private/AnimationBlueprintEditor.cpp
@@ -1172,7 +1172,7 @@ void FAnimationBlueprintEditor::Compile()
 		}
 
 		UAnimInstance* NewInstance = DebuggedMeshComponent->GetAnimInstance();
-		if((AnimBlueprint && (NewInstance->IsA(AnimBlueprint->GeneratedClass)) || (PreviewAnimBlueprint && NewInstance->IsA(PreviewAnimBlueprint->GeneratedClass))))
+		if(((AnimBlueprint && (NewInstance->IsA(AnimBlueprint->GeneratedClass))) || (PreviewAnimBlueprint && NewInstance->IsA(PreviewAnimBlueprint->GeneratedClass))))
 		{
 			PersonaUtils::SetObjectBeingDebugged(AnimBlueprint, NewInstance);
 		}
