Removed more checking code that assumes graphs are always contained within a blueprint

diff --git a/Engine/Source/Editor/AnimationBlueprintEditor/Private/AnimGraphConnectionDrawingPolicy.cpp b/Engine/Source/Editor/AnimationBlueprintEditor/Private/AnimGraphConnectionDrawingPolicy.cpp
--- a/Engine/Source/Editor/AnimationBlueprintEditor/Private/AnimGraphConnectionDrawingPolicy.cpp
+++ b/Engine/Source/Editor/AnimationBlueprintEditor/Private/AnimGraphConnectionDrawingPolicy.cpp
@@ -24,49 +24,51 @@ bool FAnimGraphConnectionDrawingPolicy::TreatWireAsExecutionPin(UEdGraphPin* Inp
 
 void FAnimGraphConnectionDrawingPolicy::BuildExecutionRoadmap()
 {
-	UAnimBlueprint* TargetBP = CastChecked<UAnimBlueprint>(FBlueprintEditorUtils::FindBlueprintForGraphChecked(GraphObj));
-	UAnimBlueprintGeneratedClass* AnimBlueprintClass = (UAnimBlueprintGeneratedClass*)(*(TargetBP->GeneratedClass));
-
-	if (TargetBP->GetObjectBeingDebugged() == NULL)
+	if(UAnimBlueprint* TargetBP = Cast<UAnimBlueprint>(FBlueprintEditorUtils::FindBlueprintForGraph(GraphObj)))
 	{
-		return;
-	}
+		UAnimBlueprintGeneratedClass* AnimBlueprintClass = (UAnimBlueprintGeneratedClass*)(*(TargetBP->GeneratedClass));
 
-	TMap<UProperty*, UObject*> PropertySourceMap;
-	AnimBlueprintClass->GetDebugData().GenerateReversePropertyMap(/*out*/ PropertySourceMap);
+		if (TargetBP->GetObjectBeingDebugged() == NULL)
+		{
+			return;
+		}
 
-	FAnimBlueprintDebugData& DebugInfo = AnimBlueprintClass->GetAnimBlueprintDebugData();
-	for (auto VisitIt = DebugInfo.UpdatedNodesThisFrame.CreateIterator(); VisitIt; ++VisitIt)
-	{
-		const FAnimBlueprintDebugData::FNodeVisit& VisitRecord = *VisitIt;
+		TMap<UProperty*, UObject*> PropertySourceMap;
+		AnimBlueprintClass->GetDebugData().GenerateReversePropertyMap(/*out*/ PropertySourceMap);
 
-		if ((VisitRecord.SourceID >= 0) && (VisitRecord.SourceID < AnimBlueprintClass->AnimNodeProperties.Num()) && (VisitRecord.TargetID >= 0) && (VisitRecord.TargetID < AnimBlueprintClass->AnimNodeProperties.Num()))
+		FAnimBlueprintDebugData& DebugInfo = AnimBlueprintClass->GetAnimBlueprintDebugData();
+		for (auto VisitIt = DebugInfo.UpdatedNodesThisFrame.CreateIterator(); VisitIt; ++VisitIt)
 		{
-			if (UAnimGraphNode_Base* SourceNode = Cast<UAnimGraphNode_Base>(PropertySourceMap.FindRef(AnimBlueprintClass->AnimNodeProperties[VisitRecord.SourceID])))
+			const FAnimBlueprintDebugData::FNodeVisit& VisitRecord = *VisitIt;
+
+			if ((VisitRecord.SourceID >= 0) && (VisitRecord.SourceID < AnimBlueprintClass->AnimNodeProperties.Num()) && (VisitRecord.TargetID >= 0) && (VisitRecord.TargetID < AnimBlueprintClass->AnimNodeProperties.Num()))
 			{
-				if (UAnimGraphNode_Base* TargetNode = Cast<UAnimGraphNode_Base>(PropertySourceMap.FindRef(AnimBlueprintClass->AnimNodeProperties[VisitRecord.TargetID])))
+				if (UAnimGraphNode_Base* SourceNode = Cast<UAnimGraphNode_Base>(PropertySourceMap.FindRef(AnimBlueprintClass->AnimNodeProperties[VisitRecord.SourceID])))
 				{
-					UEdGraphPin* PoseNet = NULL;
-
-					UAnimationGraphSchema const* AnimSchema = GetDefault<UAnimationGraphSchema>();
-					for (int32 PinIndex = 0; PinIndex < TargetNode->Pins.Num(); ++PinIndex)
+					if (UAnimGraphNode_Base* TargetNode = Cast<UAnimGraphNode_Base>(PropertySourceMap.FindRef(AnimBlueprintClass->AnimNodeProperties[VisitRecord.TargetID])))
 					{
-						UEdGraphPin* Pin = TargetNode->Pins[PinIndex];
-						check(Pin);
-						if (AnimSchema->IsPosePin(Pin->PinType) && (Pin->Direction == EGPD_Output))
+						UEdGraphPin* PoseNet = NULL;
+
+						UAnimationGraphSchema const* AnimSchema = GetDefault<UAnimationGraphSchema>();
+						for (int32 PinIndex = 0; PinIndex < TargetNode->Pins.Num(); ++PinIndex)
 						{
-							PoseNet = Pin;
-							break;
+							UEdGraphPin* Pin = TargetNode->Pins[PinIndex];
+							check(Pin);
+							if (AnimSchema->IsPosePin(Pin->PinType) && (Pin->Direction == EGPD_Output))
+							{
+								PoseNet = Pin;
+								break;
+							}
 						}
-					}
 
-					if (PoseNet != NULL)
-					{
-						//@TODO: Extend the rendering code to allow using the recorded blend weight instead of faked exec times
-						FExecPairingMap& Predecessors = PredecessorPins.FindOrAdd((UEdGraphNode*)SourceNode);
-						FTimePair& Timings = Predecessors.FindOrAdd(PoseNet);
-						Timings.PredExecTime = 0.0;
-						Timings.ThisExecTime = FMath::Clamp(VisitRecord.Weight, 0.f, 1.f);
+						if (PoseNet != NULL)
+						{
+							//@TODO: Extend the rendering code to allow using the recorded blend weight instead of faked exec times
+							FExecPairingMap& Predecessors = PredecessorPins.FindOrAdd((UEdGraphNode*)SourceNode);
+							FTimePair& Timings = Predecessors.FindOrAdd(PoseNet);
+							Timings.PredExecTime = 0.0;
+							Timings.ThisExecTime = FMath::Clamp(VisitRecord.Weight, 0.f, 1.f);
+						}
 					}
 				}
 			}
