Fix a merge issue
#rnx

diff --git a/Engine/Source/Editor/AnimationBlueprintEditor/Private/AnimationBlueprintEditorMode.cpp b/Engine/Source/Editor/AnimationBlueprintEditor/Private/AnimationBlueprintEditorMode.cpp
--- a/Engine/Source/Editor/AnimationBlueprintEditor/Private/AnimationBlueprintEditorMode.cpp
+++ b/Engine/Source/Editor/AnimationBlueprintEditor/Private/AnimationBlueprintEditorMode.cpp
@@ -134,10 +134,7 @@ FAnimationBlueprintEditorMode::FAnimationBlueprintEditorMode(const TSharedRef<FA
 	PersonaModule.RegisterPersonaViewportTabFactories(TabFactories, InAnimationBlueprintEditor, ViewportArgs);
 
 	TabFactories.RegisterFactory(PersonaModule.CreateAdvancedPreviewSceneTabFactory(InAnimationBlueprintEditor, InAnimationBlueprintEditor->GetPersonaToolkit()->GetPreviewScene()));
-#pragma message("TODO: CLEAN UP MERGED CODE")
-#if 0
-	TabFactories.RegisterFactory(PersonaModule.CreateAnimationAssetBrowserTabFactory(InAnimationBlueprintEditor, InAnimationBlueprintEditor->GetPersonaToolkit(), FOnOpenNewAsset::CreateSP(&InAnimationBlueprintEditor.Get(), &FAnimationBlueprintEditor::HandleOpenNewAsset), FOnAnimationSequenceBrowserCreated::CreateSP(&InAnimationBlueprintEditor.Get(), &FAnimationBlueprintEditor::HandleAnimationSequenceBrowserCreated), true));
-#endif
+	TabFactories.RegisterFactory(PersonaModule.CreateAnimationAssetBrowserTabFactory(InAnimationBlueprintEditor, InAnimationBlueprintEditor->GetPersonaToolkit(), FOnOpenNewAsset::CreateSP(&InAnimationBlueprintEditor.Get(), &FAnimationBlueprintEditor::HandleOpenNewAsset), FOnAnimationSequenceBrowserCreated(), true));
 	TabFactories.RegisterFactory(PersonaModule.CreateAnimBlueprintPreviewTabFactory(InAnimationBlueprintEditor, InAnimationBlueprintEditor->GetPersonaToolkit()->GetPreviewScene()));
 	TabFactories.RegisterFactory(PersonaModule.CreateAnimBlueprintAssetOverridesTabFactory(InAnimationBlueprintEditor, InAnimationBlueprintEditor->GetPersonaToolkit()->GetAnimBlueprint(), InAnimationBlueprintEditor->OnPostUndo));
 	TabFactories.RegisterFactory(PersonaModule.CreateSkeletonSlotNamesTabFactory(InAnimationBlueprintEditor, InAnimationBlueprintEditor->GetSkeletonTree()->GetEditableSkeleton(), InAnimationBlueprintEditor->OnPostUndo, FOnObjectSelected::CreateSP(&InAnimationBlueprintEditor.Get(), &FAnimationBlueprintEditor::HandleObjectSelected)));
