Removing duplicate Hide Unrelated Nodes from the Anim BP editor

diff --git a/Engine/Source/Editor/AnimGraph/Private/AnimGraphCommands.cpp b/Engine/Source/Editor/AnimGraph/Private/AnimGraphCommands.cpp
--- a/Engine/Source/Editor/AnimGraph/Private/AnimGraphCommands.cpp
+++ b/Engine/Source/Editor/AnimGraph/Private/AnimGraphCommands.cpp
@@ -7,7 +7,6 @@
 void FAnimGraphCommands::RegisterCommands()
 {
 	UI_COMMAND(TogglePoseWatch, "Toggle Pose Watch", "Toggle Pose Watching on this node", EUserInterfaceActionType::Button, FInputChord());
-	UI_COMMAND( ToggleHideUnrelatedNodes, "Hide Unrelated", "Toggles automatically hiding nodes which are unrelated to the selected nodes.", EUserInterfaceActionType::ToggleButton, FInputChord() );
 }
 
 #undef LOCTEXT_NAMESPACE
diff --git a/Engine/Source/Editor/AnimGraph/Public/AnimGraphCommands.h b/Engine/Source/Editor/AnimGraph/Public/AnimGraphCommands.h
--- a/Engine/Source/Editor/AnimGraph/Public/AnimGraphCommands.h
+++ b/Engine/Source/Editor/AnimGraph/Public/AnimGraphCommands.h
@@ -21,7 +21,5 @@ public:
 
 	// Toggle pose watching for a given node
 	TSharedPtr<FUICommandInfo> TogglePoseWatch;
-	
-	// Toggle hiding nodes which are not related to the selected nodes
-	TSharedPtr< FUICommandInfo > ToggleHideUnrelatedNodes;
+
 };
diff --git a/Engine/Source/Editor/AnimationBlueprintEditor/Private/AnimationBlueprintEditor.cpp b/Engine/Source/Editor/AnimationBlueprintEditor/Private/AnimationBlueprintEditor.cpp
--- a/Engine/Source/Editor/AnimationBlueprintEditor/Private/AnimationBlueprintEditor.cpp
+++ b/Engine/Source/Editor/AnimationBlueprintEditor/Private/AnimationBlueprintEditor.cpp
@@ -370,39 +370,6 @@ void FAnimationBlueprintEditor::ExtendToolbar()
 			}
 		));
 	}
-
-	struct Local
-	{
-		static void FillToolbar(FToolBarBuilder& ToolbarBuilder, const TSharedRef< FUICommandList > ToolkitCommands, FAnimationBlueprintEditor* BlueprintEditor)
-		{
-			ToolbarBuilder.BeginSection("Graph");
-			{
-				ToolbarBuilder.AddToolBarButton(
-					FAnimGraphCommands::Get().ToggleHideUnrelatedNodes,
-					NAME_None,
-					TAttribute<FText>(),
-					TAttribute<FText>(),
-					FSlateIcon(FEditorStyle::GetStyleSetName(), "GraphEditor.ToggleHideUnrelatedNodes")
-				);
-				ToolbarBuilder.AddComboButton(
-					FUIAction(),
-					FOnGetContent::CreateSP(BlueprintEditor, &FBlueprintEditor::MakeHideUnrelatedNodesOptionsMenu),
-					LOCTEXT("HideUnrelatedNodesOptions", "Hide Unrelated Nodes Options"),
-					LOCTEXT("HideUnrelatedNodesOptionsMenu", "Hide Unrelated Nodes options menu"),
-					TAttribute<FSlateIcon>(),
-					true
-				);
-			}
-			ToolbarBuilder.EndSection();
-		}
-	};
-
-	ToolbarExtender->AddToolBarExtension(
-		"Asset",
-		EExtensionHook::After,
-		GetToolkitCommands(),
-		FToolBarExtensionDelegate::CreateStatic( &Local::FillToolbar, GetToolkitCommands(), this )
-	);
 }
 
 UBlueprint* FAnimationBlueprintEditor::GetBlueprintObj() const
@@ -465,12 +432,6 @@ void FAnimationBlueprintEditor::CreateDefaultCommands()
 {
 	{
 		FBlueprintEditor::CreateDefaultCommands();
-
-		ToolkitCommands->MapAction(
-			FAnimGraphCommands::Get().ToggleHideUnrelatedNodes,
-			FExecuteAction::CreateSP(this, &FBlueprintEditor::ToggleHideUnrelatedNodes),
-			FCanExecuteAction(),
-			FIsActionChecked::CreateSP(this, &FBlueprintEditor::IsToggleHideUnrelatedNodesChecked));
 	}
 }
 
