Fix anim BP debugger disconnection when compiling-on-save

diff --git a/Engine/Source/Editor/AnimationBlueprintEditor/Private/AnimationBlueprintEditor.cpp b/Engine/Source/Editor/AnimationBlueprintEditor/Private/AnimationBlueprintEditor.cpp
--- a/Engine/Source/Editor/AnimationBlueprintEditor/Private/AnimationBlueprintEditor.cpp
+++ b/Engine/Source/Editor/AnimationBlueprintEditor/Private/AnimationBlueprintEditor.cpp
@@ -161,7 +161,8 @@ public:
 // FAnimationBlueprintEditor
 
 FAnimationBlueprintEditor::FAnimationBlueprintEditor()
-	: PersonaMeshDetailLayout(NULL)
+	: PersonaMeshDetailLayout(nullptr)
+	, DebuggedMeshComponent(nullptr)
 {
 	GEditor->OnBlueprintPreCompile().AddRaw(this, &FAnimationBlueprintEditor::OnBlueprintPreCompile);
 	LastGraphPinType.ResetToDefaults();
@@ -242,6 +243,9 @@ void FAnimationBlueprintEditor::InitAnimationBlueprintEditor(const EToolkitMode:
 		SkeletonTree = SkeletonEditorModule.CreateSkeletonTree(PersonaToolkit->GetSkeleton(), SkeletonTreeArgs);
 	}
 
+	// Register for compilation events
+	InAnimBlueprint->OnCompiled().AddSP(this, &FAnimationBlueprintEditor::OnBlueprintPostCompile);
+
 	// Build up a list of objects being edited in this asset editor
 	TArray<UObject*> ObjectsBeingEdited;
 	ObjectsBeingEdited.Add(InAnimBlueprint);
@@ -1089,56 +1093,6 @@ void FAnimationBlueprintEditor::RecompileAnimBlueprintIfDirty()
 	}
 }
 
-void FAnimationBlueprintEditor::Compile()
-{
-	// Grab the currently debugged object, so we can re-set it below
-	USkeletalMeshComponent* DebuggedMeshComponent = nullptr;
-	
-	if (UBlueprint* Blueprint = GetBlueprintObj())
-	{
-		UAnimInstance* CurrentDebugObject = Cast<UAnimInstance>(Blueprint->GetObjectBeingDebugged());
-		if(CurrentDebugObject)
-		{
-			// Force close any asset editors that are using the AnimScriptInstance (such as the Property Matrix), the class will be garbage collected
-			GEditor->GetEditorSubsystem<UAssetEditorSubsystem>()->CloseOtherEditors(CurrentDebugObject, nullptr);
-			DebuggedMeshComponent = CurrentDebugObject->GetSkelMeshComponent();
-		}
-	}
-
-	// Compile the blueprint
-	FBlueprintEditor::Compile();
-
-	if (DebuggedMeshComponent != nullptr)
-	{
-		if (DebuggedMeshComponent->GetAnimInstance() == nullptr)
-		{
-			// try reinitialize animation if it doesn't exist
-			DebuggedMeshComponent->InitAnim(true);
-		}
-
-		// re-apply preview anim bp if needed
-		UAnimBlueprint* AnimBlueprint = GetAnimBlueprint();
-		UAnimBlueprint* PreviewAnimBlueprint = AnimBlueprint ? AnimBlueprint->GetPreviewAnimationBlueprint() : nullptr;
-		
-		if (PreviewAnimBlueprint)
-		{
-			PersonaToolkit->GetPreviewScene()->SetPreviewAnimationBlueprint(PreviewAnimBlueprint, AnimBlueprint);
-		}
-
-		UAnimInstance* NewInstance = DebuggedMeshComponent->GetAnimInstance();
-		if ((AnimBlueprint && NewInstance->IsA(AnimBlueprint->GeneratedClass)) || (PreviewAnimBlueprint && NewInstance->IsA(PreviewAnimBlueprint->GeneratedClass)))
-		{
-			PersonaUtils::SetObjectBeingDebugged(AnimBlueprint, NewInstance);
-		}
-	}
-
-	// reset the selected skeletal control node
-	SelectedAnimGraphNode.Reset();
-
-	// if the user manipulated Pin values directly from the node, then should copy updated values to the internal node to retain data consistency
-	OnPostCompile();
-}
-
 FName FAnimationBlueprintEditor::GetToolkitFName() const
 {
 	return FName("AnimationBlueprintEditor");
@@ -1412,6 +1366,59 @@ void FAnimationBlueprintEditor::OnBlueprintPreCompile(UBlueprint* BlueprintToCom
 			}
 		}
 	}
+
+	if(BlueprintToCompile == GetBlueprintObj())
+	{
+		// Grab the currently debugged object, so we can re-set it below in OnBlueprintPostCompile
+		DebuggedMeshComponent = nullptr;
+
+		UAnimInstance* CurrentDebugObject = Cast<UAnimInstance>(BlueprintToCompile->GetObjectBeingDebugged());
+		if(CurrentDebugObject)
+		{
+			// Force close any asset editors that are using the AnimScriptInstance (such as the Property Matrix), the class will be garbage collected
+			GEditor->GetEditorSubsystem<UAssetEditorSubsystem>()->CloseOtherEditors(CurrentDebugObject, nullptr);
+			DebuggedMeshComponent = CurrentDebugObject->GetSkelMeshComponent();
+		}
+	}
+}
+
+void FAnimationBlueprintEditor::OnBlueprintPostCompile(UBlueprint* InBlueprint)
+{
+	if(InBlueprint == GetBlueprintObj())
+	{
+		if (DebuggedMeshComponent != nullptr)
+		{
+			if (DebuggedMeshComponent->GetAnimInstance() == nullptr)
+			{
+				// try reinitialize animation if it doesn't exist
+				DebuggedMeshComponent->InitAnim(true);
+			}
+
+			// re-apply preview anim bp if needed
+			UAnimBlueprint* AnimBlueprint = GetAnimBlueprint();
+			UAnimBlueprint* PreviewAnimBlueprint = AnimBlueprint ? AnimBlueprint->GetPreviewAnimationBlueprint() : nullptr;
+		
+			if (PreviewAnimBlueprint)
+			{
+				PersonaToolkit->GetPreviewScene()->SetPreviewAnimationBlueprint(PreviewAnimBlueprint, AnimBlueprint);
+			}
+
+			UAnimInstance* NewInstance = DebuggedMeshComponent->GetAnimInstance();
+			if ((AnimBlueprint && NewInstance->IsA(AnimBlueprint->GeneratedClass)) || (PreviewAnimBlueprint && NewInstance->IsA(PreviewAnimBlueprint->GeneratedClass)))
+			{
+				PersonaUtils::SetObjectBeingDebugged(AnimBlueprint, NewInstance);
+			}
+		}
+
+		// reset the selected skeletal control node
+		SelectedAnimGraphNode.Reset();
+
+		// if the user manipulated Pin values directly from the node, then should copy updated values to the internal node to retain data consistency
+		OnPostCompile();
+
+		// We dont cache this persistently, only during a pre/post compile bracket
+		DebuggedMeshComponent = nullptr;
+	}
 }
 
 void FAnimationBlueprintEditor::OnBlueprintChangedImpl(UBlueprint* InBlueprint, bool bIsJustBeingCompiled /*= false*/)
diff --git a/Engine/Source/Editor/AnimationBlueprintEditor/Private/AnimationBlueprintEditor.h b/Engine/Source/Editor/AnimationBlueprintEditor/Private/AnimationBlueprintEditor.h
--- a/Engine/Source/Editor/AnimationBlueprintEditor/Private/AnimationBlueprintEditor.h
+++ b/Engine/Source/Editor/AnimationBlueprintEditor/Private/AnimationBlueprintEditor.h
@@ -177,7 +177,6 @@ protected:
 	virtual bool CanAddPosePin() const override;
 	virtual void OnRemovePosePin() override;
 	virtual bool CanRemovePosePin() const override;
-	virtual void Compile() override;
 	virtual void OnGraphEditorFocused(const TSharedRef<class SGraphEditor>& InGraphEditor) override;
 	virtual void OnGraphEditorBackgrounded(const TSharedRef<SGraphEditor>& InGraphEditor) override;
 	virtual void OnConvertToSequenceEvaluator() override;
@@ -264,6 +263,9 @@ private:
 	/** Called immediately prior to a blueprint compilation */
 	void OnBlueprintPreCompile(UBlueprint* BlueprintToCompile);
 
+	/** Called immediately after to a blueprint compilation */
+	void OnBlueprintPostCompile(UBlueprint* InBlueprint);
+
 	/** Called post compile to copy node data */
 	void OnPostCompile();
 
@@ -323,4 +325,7 @@ private:
 
     /** Configuration class used to store editor settings across sessions. */
 	UAnimationBlueprintEditorOptions* EditorOptions;
+
+	/** Cached mesh component held during compilation, used to reconnect debugger */
+	USkeletalMeshComponent* DebuggedMeshComponent;
 };
