Transition nodes for outgoing transitions are now highlighted (along with the transition arrows) when the mouse is over a state node

diff --git a/Engine/Source/Editor/AnimationBlueprintEditor/Private/AnimationStateNodes/SGraphNodeAnimState.cpp b/Engine/Source/Editor/AnimationBlueprintEditor/Private/AnimationStateNodes/SGraphNodeAnimState.cpp
--- a/Engine/Source/Editor/AnimationBlueprintEditor/Private/AnimationStateNodes/SGraphNodeAnimState.cpp
+++ b/Engine/Source/Editor/AnimationBlueprintEditor/Private/AnimationStateNodes/SGraphNodeAnimState.cpp
@@ -15,6 +15,7 @@
 #include "Animation/AnimNode_StateMachine.h"
 #include "AnimGraphNode_StateMachineBase.h"
 #include "Widgets/Text/SInlineEditableTextBlock.h"
+#include "SGraphPanel.h"
 
 #define LOCTEXT_NAMESPACE "SGraphNodeAnimState"
 
@@ -172,6 +173,46 @@ FSlateColor SGraphNodeAnimState::GetBorderBackgroundColor() const
 	return InactiveStateColor;
 }
 
+void SGraphNodeAnimState::OnMouseEnter(const FGeometry& MyGeometry, const FPointerEvent& MouseEvent)
+{
+	// Add pins to the hover set so outgoing transitions arrows remains highlighted while the mouse is over the state node
+	if (const UAnimStateNodeBase* StateNode = Cast<UAnimStateNodeBase>(GraphNode))
+	{
+		if (const UEdGraphPin* OutputPin = StateNode->GetOutputPin())
+		{
+			TSharedPtr<SGraphPanel> OwnerPanel = GetOwnerPanel();
+			check(OwnerPanel.IsValid());
+
+			for (int32 LinkIndex = 0; LinkIndex < OutputPin->LinkedTo.Num(); ++LinkIndex)
+			{
+				OwnerPanel->AddPinToHoverSet(OutputPin->LinkedTo[LinkIndex]);
+			}
+		}
+	}
+	
+	SGraphNode::OnMouseEnter(MyGeometry, MouseEvent);
+}
+
+void SGraphNodeAnimState::OnMouseLeave(const FPointerEvent& MouseEvent)
+{
+	// Remove manually added pins from the hover set
+	if (const UAnimStateNodeBase* StateNode = Cast<UAnimStateNodeBase>(GraphNode))
+	{
+		if(const UEdGraphPin* OutputPin = StateNode->GetOutputPin())
+		{
+			TSharedPtr<SGraphPanel> OwnerPanel = GetOwnerPanel();
+			check(OwnerPanel.IsValid());
+
+			for (int32 LinkIndex = 0; LinkIndex < OutputPin->LinkedTo.Num(); ++LinkIndex)
+			{
+				OwnerPanel->RemovePinFromHoverSet(OutputPin->LinkedTo[LinkIndex]);
+			}
+		}
+	}
+
+	SGraphNode::OnMouseLeave(MouseEvent);
+}
+
 void SGraphNodeAnimState::UpdateGraphNode()
 {
 	InputPins.Empty();
diff --git a/Engine/Source/Editor/AnimationBlueprintEditor/Private/AnimationStateNodes/SGraphNodeAnimState.h b/Engine/Source/Editor/AnimationBlueprintEditor/Private/AnimationStateNodes/SGraphNodeAnimState.h
--- a/Engine/Source/Editor/AnimationBlueprintEditor/Private/AnimationStateNodes/SGraphNodeAnimState.h
+++ b/Engine/Source/Editor/AnimationBlueprintEditor/Private/AnimationStateNodes/SGraphNodeAnimState.h
@@ -39,6 +39,11 @@ public:
 	virtual TSharedPtr<SToolTip> GetComplexTooltip() override;
 	// End of SGraphNode interface
 
+	// SWidget interface
+	void OnMouseEnter(const FGeometry& MyGeometry, const FPointerEvent& MouseEvent) override;
+	void OnMouseLeave(const FPointerEvent& MouseEvent) override;
+	// End of SWidget interface
+
 	static void GetStateInfoPopup(UEdGraphNode* GraphNode, TArray<FGraphInformationPopupInfo>& Popups);
 protected:
 	FSlateColor GetBorderBackgroundColor() const;
diff --git a/Engine/Source/Editor/AnimationBlueprintEditor/Private/AnimationStateNodes/SGraphNodeAnimTransition.cpp b/Engine/Source/Editor/AnimationBlueprintEditor/Private/AnimationStateNodes/SGraphNodeAnimTransition.cpp
--- a/Engine/Source/Editor/AnimationBlueprintEditor/Private/AnimationStateNodes/SGraphNodeAnimTransition.cpp
+++ b/Engine/Source/Editor/AnimationBlueprintEditor/Private/AnimationStateNodes/SGraphNodeAnimTransition.cpp
@@ -82,6 +82,8 @@ void SGraphNodeAnimTransition::PerformSecondPassLayout(const TMap< UObject*, TSh
 
 			TransIndex = Transitions.IndexOfByKey(TransNode);
 			NumOfTrans = Transitions.Num();
+
+			PrevStateNodeWidgetPtr = PrevNodeWidget;
 		}
 	}
 
@@ -318,9 +320,10 @@ FLinearColor SGraphNodeAnimTransition::StaticGetTransitionColor(UAnimStateTransi
 }
 
 FSlateColor SGraphNodeAnimTransition::GetTransitionColor() const
-{
+{	
+	// Highlight the transition node when the node is hovered or when the previous state is hovered
 	UAnimStateTransitionNode* TransNode = CastChecked<UAnimStateTransitionNode>(GraphNode);
-	return StaticGetTransitionColor(TransNode, IsHovered());
+	return StaticGetTransitionColor(TransNode, (IsHovered() || (PrevStateNodeWidgetPtr.IsValid() && PrevStateNodeWidgetPtr.Pin()->IsHovered())));
 }
 
 const FSlateBrush* SGraphNodeAnimTransition::GetTransitionIconImage() const
diff --git a/Engine/Source/Editor/AnimationBlueprintEditor/Private/AnimationStateNodes/SGraphNodeAnimTransition.h b/Engine/Source/Editor/AnimationBlueprintEditor/Private/AnimationStateNodes/SGraphNodeAnimTransition.h
--- a/Engine/Source/Editor/AnimationBlueprintEditor/Private/AnimationStateNodes/SGraphNodeAnimTransition.h
+++ b/Engine/Source/Editor/AnimationBlueprintEditor/Private/AnimationStateNodes/SGraphNodeAnimTransition.h
@@ -44,6 +44,9 @@ public:
 private:
 	TSharedPtr<STextEntryPopup> TextEntryWidget;
 
+	/** Cache of the widget representing the previous state node */
+	mutable TWeakPtr<SNode> PrevStateNodeWidgetPtr;
+
 private:
 	FText GetPreviewCornerText(bool reverse) const;
 	FSlateColor GetTransitionColor() const;
