Fix missing Animation Blueprint editor buttons
#fix Incorrect base-class virtual was being called

diff --git a/Engine/Source/Editor/AnimationBlueprintEditor/Private/AnimationBlueprintEditor.cpp b/Engine/Source/Editor/AnimationBlueprintEditor/Private/AnimationBlueprintEditor.cpp
--- a/Engine/Source/Editor/AnimationBlueprintEditor/Private/AnimationBlueprintEditor.cpp
+++ b/Engine/Source/Editor/AnimationBlueprintEditor/Private/AnimationBlueprintEditor.cpp
@@ -1123,7 +1123,7 @@ FLinearColor FAnimationBlueprintEditor::GetWorldCentricTabColorScale() const
 
 void FAnimationBlueprintEditor::InitToolMenuContext(FToolMenuContext& MenuContext)
 {
-	FAssetEditorToolkit::InitToolMenuContext(MenuContext);
+	IAnimationBlueprintEditor::InitToolMenuContext(MenuContext);
 
 	UPersonaToolMenuContext* Context = NewObject<UPersonaToolMenuContext>();
 	Context->SetToolkit(GetPersonaToolkit());
