Removes old AnimationBlueprint Toolbar as part of the default Layout

diff --git a/Engine/Source/Editor/AnimationBlueprintEditor/Private/AnimationBlueprintInterfaceEditorMode.cpp b/Engine/Source/Editor/AnimationBlueprintEditor/Private/AnimationBlueprintInterfaceEditorMode.cpp
--- a/Engine/Source/Editor/AnimationBlueprintEditor/Private/AnimationBlueprintInterfaceEditorMode.cpp
+++ b/Engine/Source/Editor/AnimationBlueprintEditor/Private/AnimationBlueprintInterfaceEditorMode.cpp
@@ -22,14 +22,6 @@ FAnimationBlueprintInterfaceEditorMode::FAnimationBlueprintInterfaceEditorMode(c
 			FTabManager::NewPrimaryArea()
 			->SetOrientation(Orient_Vertical)
 			->Split
-			(
-				// Top toolbar
-				FTabManager::NewStack() 
-				->SetSizeCoefficient(0.186721f)
-				->SetHideTabWell(true)
-				->AddTab(InAnimationBlueprintEditor->GetToolbarTabId(), ETabState::OpenedTab)
-			)
-			->Split
 			(
 				// Main application area
 				FTabManager::NewSplitter()
