Fixed non-unity non-PCH compile errors

diff --git a/Engine/Source/Editor/AnimGraph/Public/AnimGraphAttributes.h b/Engine/Source/Editor/AnimGraph/Public/AnimGraphAttributes.h
--- a/Engine/Source/Editor/AnimGraph/Public/AnimGraphAttributes.h
+++ b/Engine/Source/Editor/AnimGraph/Public/AnimGraphAttributes.h
@@ -3,6 +3,8 @@
 #pragma once
 
 #include "CoreMinimal.h"
+#include "Styling/SlateBrush.h"
+#include "Styling/SlateColor.h"
 #include "AnimGraphAttributes.generated.h"
 
 UENUM()
diff --git a/Engine/Source/Editor/AnimationBlueprintEditor/Private/AnimationPins/SGraphPinPose.cpp b/Engine/Source/Editor/AnimationBlueprintEditor/Private/AnimationPins/SGraphPinPose.cpp
--- a/Engine/Source/Editor/AnimationBlueprintEditor/Private/AnimationPins/SGraphPinPose.cpp
+++ b/Engine/Source/Editor/AnimationBlueprintEditor/Private/AnimationPins/SGraphPinPose.cpp
@@ -6,6 +6,10 @@
 #include "AnimGraphAttributes.h"
 #include "Kismet2/BlueprintEditorUtils.h"
 #include "Widgets/Layout/SWrapBox.h"
+#include "Widgets/Layout/SBorder.h"
+#include "Widgets/Images/SImage.h"
+#include "Widgets/SOverlay.h"
+#include "Widgets/SBoxPanel.h"
 #include "SGraphPanel.h"
 
 #define LOCTEXT_NAMESPACE "SGraphPinPose"
