BlendSpace 2.0: Blendspace Graph Node

Added a new animation graph node that hosts its own UBlendSpaceBase. Modified UBlendSpaceBase to allow for pose links to be evaluated as the sample points.
The new blend space graphs can be spawned from existing UBlendSpace and UBlendSpace1D assets, or they can be created from scratch, or they can be converted from existing blendspace player nodes via the context menu.

Fixed anim node conversion functions so that their transactions work correctly.

Updated FBlueprintEditorUtils::IsGraphNameUnique to allow it to work with any object as the outer, not just UBlueprint. UBlueprint still has a special case for functions and events. This is to support GenerateUniqueGraphName within a scope (e.g. an outer graph).

Formalized the concept of 'node sub-graphs' (as well as the composite node pattern a little). Previously a number of known node types that contained sub-graphs (e.g. UK2Node_Composite) had special case logic for dealing with node/graph deletion etc. Now  any node can opt into this behaviour via the GetSubGraphs() override.

Added status bar readouts for the blendspace grid, so we dont have to stuff the prompts into the tooltip any more.

Moved anim BP related APIs out of FBlueprintEditor. They are always used via FAnimationBlueprintEditor.

Refactored graph title bar widget creation out into a function to allow other document tab factories to create it.

Altered breadcrumb trail click callbacks and SMyBlueprint::ExecuteAction to always JumpToHyperLink rather than calling OpenDocument directly. This allows unknown (to FBlueprintEditor) document types that reference objects to be correctly jumped to using the breadcrumb trail. Derived asset editors (i.e. FAnimationBlueprintEditor) can intercept the JumpToHyperlink call to ensure that the correct document is presented (i.e. the correct tab payload is generated).

Instead of making yet another bunch of duplicated code for handling the various alpha blend options, refactored this into FAnimGraphNodeAlphaOptions (for editor code) and FAnimNodeAlphaOptions (for runtime code).

Added OnCopyTermDefaultsToDefaultObject for per-node copying of default values from editor node to runtime node, rather than another special-case in the compiler.

diff --git a/Engine/Source/Editor/AnimGraph/Private/AnimBlueprintCompilationContext.h b/Engine/Source/Editor/AnimGraph/Private/AnimBlueprintCompilationContext.h
--- a/Engine/Source/Editor/AnimGraph/Private/AnimBlueprintCompilationContext.h
+++ b/Engine/Source/Editor/AnimGraph/Private/AnimBlueprintCompilationContext.h
@@ -73,6 +73,29 @@ private:
 	FAnimBlueprintCompilerContext* CompilerContext;
 };
 
+class FAnimBlueprintNodeCopyTermDefaultsContext : public IAnimBlueprintNodeCopyTermDefaultsContext
+{
+private:
+	friend class FAnimBlueprintCompilerContext;
+
+	FAnimBlueprintNodeCopyTermDefaultsContext(const FProperty* InTargetProperty, uint8* InDestinationPtr, const uint8* InSourcePtr, int32 InNodePropertyIndex)
+		: TargetProperty(InTargetProperty)
+		, DestinationPtr(InDestinationPtr)
+		, SourcePtr(InSourcePtr)
+		, NodePropertyIndex(InNodePropertyIndex)
+	{}
+
+	virtual const FProperty* GetTargetPropertyImpl() const override { return TargetProperty; }
+	virtual uint8* GetDestinationPtrImpl() const override { return DestinationPtr; }
+	virtual const uint8* GetSourcePtrImpl() const override { return SourcePtr; }
+	virtual int32 GetNodePropertyIndexImpl() const override { return NodePropertyIndex; }
+
+	const FProperty* TargetProperty;
+	uint8* DestinationPtr;
+	const uint8* SourcePtr;
+	int32 NodePropertyIndex;
+};
+
 class FAnimBlueprintPostExpansionStepContext : public IAnimBlueprintPostExpansionStepContext
 {
 private:
diff --git a/Engine/Source/Editor/AnimGraph/Private/AnimBlueprintCompiler.cpp b/Engine/Source/Editor/AnimGraph/Private/AnimBlueprintCompiler.cpp
--- a/Engine/Source/Editor/AnimGraph/Private/AnimBlueprintCompiler.cpp
+++ b/Engine/Source/Editor/AnimGraph/Private/AnimBlueprintCompiler.cpp
@@ -545,6 +545,9 @@ void FAnimBlueprintCompilerContext::CopyTermDefaultsToDefaultObject(UObject* Def
 		TMap<UAnimGraphNode_Base*, int32> LinkIndexMap;
 		TMap<UAnimGraphNode_Base*, uint8*> NodeBaseAddresses;
 
+		FAnimBlueprintGeneratedClassCompiledData CompiledData(NewAnimBlueprintClass);
+		FAnimBlueprintCopyTermDefaultsContext CompilerContext(this);
+
 		// Initialize animation nodes from their templates
 		for (TFieldIterator<FProperty> It(DefaultAnimInstance->GetClass(), EFieldIteratorFlags::ExcludeSuper); It; ++It)
 		{
@@ -557,40 +560,10 @@ void FAnimBlueprintCompilerContext::CopyTermDefaultsToDefaultObject(UObject* Def
 				check(CastFieldChecked<FStructProperty>(TargetProperty)->Struct == SourceNodeProperty->Struct);
 
 				uint8* DestinationPtr = TargetProperty->ContainerPtrToValuePtr<uint8>(DefaultAnimInstance);
-				uint8* SourcePtr = SourceNodeProperty->ContainerPtrToValuePtr<uint8>(VisualAnimNode);
+				const uint8* SourcePtr = SourceNodeProperty->ContainerPtrToValuePtr<uint8>(VisualAnimNode);
 
-				if(UAnimGraphNode_Root* RootNode = ExactCast<UAnimGraphNode_Root>(VisualAnimNode))
-				{
-					// patch graph name into root nodes
-					FAnimNode_Root NewRoot = *reinterpret_cast<FAnimNode_Root*>(SourcePtr);
-					NewRoot.Name = Cast<UAnimGraphNode_Root>(MessageLog.FindSourceObject(RootNode))->GetGraph()->GetFName();
-					TargetProperty->CopyCompleteValue(DestinationPtr, &NewRoot);
-				}
-				else if(UAnimGraphNode_LinkedInputPose* LinkedInputPoseNode = ExactCast<UAnimGraphNode_LinkedInputPose>(VisualAnimNode))
-				{
-					// patch graph name into linked input pose nodes
-					FAnimNode_LinkedInputPose NewLinkedInputPose = *reinterpret_cast<FAnimNode_LinkedInputPose*>(SourcePtr);
-					NewLinkedInputPose.Graph = Cast<UAnimGraphNode_LinkedInputPose>(MessageLog.FindSourceObject(LinkedInputPoseNode))->GetGraph()->GetFName();
-					TargetProperty->CopyCompleteValue(DestinationPtr, &NewLinkedInputPose);
-				}
-				else if(UAnimGraphNode_LinkedAnimGraph* LinkedAnimGraphNode = ExactCast<UAnimGraphNode_LinkedAnimGraph>(VisualAnimNode))
-				{
-					// patch node index into linked anim graph nodes
-					FAnimNode_LinkedAnimGraph NewLinkedAnimGraph = *reinterpret_cast<FAnimNode_LinkedAnimGraph*>(SourcePtr);
-					NewLinkedAnimGraph.NodeIndex = LinkIndexCount;
-					TargetProperty->CopyCompleteValue(DestinationPtr, &NewLinkedAnimGraph);
-				}
-				else if(UAnimGraphNode_LinkedAnimLayer* LinkedAnimLayerNode = ExactCast<UAnimGraphNode_LinkedAnimLayer>(VisualAnimNode))
-				{
-					// patch node index into linked anim layer nodes
-					FAnimNode_LinkedAnimLayer NewLinkedAnimLayer = *reinterpret_cast<FAnimNode_LinkedAnimLayer*>(SourcePtr);
-					NewLinkedAnimLayer.NodeIndex = LinkIndexCount;
-					TargetProperty->CopyCompleteValue(DestinationPtr, &NewLinkedAnimLayer);
-				}
-				else
-				{
-					TargetProperty->CopyCompleteValue(DestinationPtr, SourcePtr);
-				}
+				FAnimBlueprintNodeCopyTermDefaultsContext NodeContext(TargetProperty, DestinationPtr, SourcePtr, LinkIndexCount);
+				VisualAnimNode->CopyTermDefaultsToDefaultObject(CompilerContext, NodeContext, CompiledData);
 
 				LinkIndexMap.Add(VisualAnimNode, LinkIndexCount);
 				NodeBaseAddresses.Add(VisualAnimNode, DestinationPtr);
@@ -619,8 +592,6 @@ void FAnimBlueprintCompilerContext::CopyTermDefaultsToDefaultObject(UObject* Def
 			}
 		}   
 
-		FAnimBlueprintGeneratedClassCompiledData CompiledData(NewAnimBlueprintClass);
-		FAnimBlueprintCopyTermDefaultsContext CompilerContext(this);
 		OnCopyTermDefaultsToDefaultObjectDelegate.Broadcast(DefaultObject, CompilerContext, CompiledData);
 
 		UAnimBlueprintGeneratedClass* AnimBlueprintGeneratedClass = CastChecked<UAnimBlueprintGeneratedClass>(NewClass);
diff --git a/Engine/Source/Editor/AnimGraph/Private/AnimBlueprintCompilerHandler_Attributes.cpp b/Engine/Source/Editor/AnimGraph/Private/AnimBlueprintCompilerHandler_Attributes.cpp
--- a/Engine/Source/Editor/AnimGraph/Private/AnimBlueprintCompilerHandler_Attributes.cpp
+++ b/Engine/Source/Editor/AnimGraph/Private/AnimBlueprintCompilerHandler_Attributes.cpp
@@ -18,6 +18,10 @@
 #include "AnimGraphNode_CustomTransitionResult.h"
 #include "AnimationCustomTransitionGraph.h"
 #include "AnimStateTransitionNode.h"
+#include "AnimGraphNode_BlendSpaceSampleResult.h"
+#include "BlendSpaceGraph.h"
+#include "AnimGraphNode_BlendSpaceGraphBase.h"
+#include "AnimationBlendSpaceSampleGraph.h"
 
 FAnimBlueprintCompilerHandler_Attributes::FAnimBlueprintCompilerHandler_Attributes(IAnimBlueprintCompilerCreationContext& InCreationContext)
 {
@@ -235,6 +239,14 @@ void FAnimBlueprintCompilerHandler_Attributes::PreProcessAnimationNodes(TArrayVi
 				UAnimGraphNode_StateMachineBase* StateMachineNode = CastChecked<UAnimGraphNode_StateMachineBase>(StateMachineGraph->GetOuter());
 				return TraverseNodes_Recursive_PerNode(StateMachineNode, InAttributes);
 			}
+			else if(UAnimGraphNode_BlendSpaceSampleResult* SampleResultNode = Cast<UAnimGraphNode_BlendSpaceSampleResult>(InNode))
+			{
+				UAnimationBlendSpaceSampleGraph* SampleGraph = CastChecked<UAnimationBlendSpaceSampleGraph>(SampleResultNode->GetGraph());
+				UBlendSpaceGraph* BlendSpaceGraph = CastChecked<UBlendSpaceGraph>(SampleGraph->GetOuter());
+				UAnimGraphNode_BlendSpaceGraphBase* BlendSpaceGraphNode = CastChecked<UAnimGraphNode_BlendSpaceGraphBase>(BlendSpaceGraph->GetOuter());
+
+				return TraverseNodes_Recursive_PerNode(BlendSpaceGraphNode, InAttributes);
+			}
 
 			return UAnimGraphNode_Base::FNodeAttributeArray();
 		}
diff --git a/Engine/Source/Editor/AnimGraph/Private/AnimBlueprintGeneratedClassCompiledData.cpp b/Engine/Source/Editor/AnimGraph/Private/AnimBlueprintGeneratedClassCompiledData.cpp
--- a/Engine/Source/Editor/AnimGraph/Private/AnimBlueprintGeneratedClassCompiledData.cpp
+++ b/Engine/Source/Editor/AnimGraph/Private/AnimBlueprintGeneratedClassCompiledData.cpp
@@ -69,4 +69,11 @@ FAnimBlueprintDebugData& FAnimBlueprintGeneratedClassCompiledData::GetAnimBluepr
 TMap<FName, FGraphAssetPlayerInformation>& FAnimBlueprintGeneratedClassCompiledData::GetGraphAssetPlayerInformation() const
 {
 	return Class->GraphAssetPlayerInformation;
+}
+
+UBlendSpaceBase* FAnimBlueprintGeneratedClassCompiledData::AddBlendSpace(UBlendSpaceBase* InSourceBlendSpace)
+{
+	UBlendSpaceBase* CopiedBlendSpace = DuplicateObject(InSourceBlendSpace, Class);
+	Class->BlendSpaces.Add(CopiedBlendSpace);
+	return CopiedBlendSpace;
 }
\ No newline at end of file
diff --git a/Engine/Source/Editor/AnimGraph/Private/AnimBlueprintGeneratedClassCompiledData.h b/Engine/Source/Editor/AnimGraph/Private/AnimBlueprintGeneratedClassCompiledData.h
--- a/Engine/Source/Editor/AnimGraph/Private/AnimBlueprintGeneratedClassCompiledData.h
+++ b/Engine/Source/Editor/AnimGraph/Private/AnimBlueprintGeneratedClassCompiledData.h
@@ -24,6 +24,7 @@ private:
 	virtual FPropertyAccessLibrary& GetPropertyAccessLibrary() const override;
 	virtual FAnimBlueprintDebugData& GetAnimBlueprintDebugData() const override;
 	virtual TMap<FName, FGraphAssetPlayerInformation>& GetGraphAssetPlayerInformation() const override;
+	virtual UBlendSpaceBase* AddBlendSpace(UBlendSpaceBase* InSourceBlendSpace) override;
 
 private:
 	// The class we are wrapping
diff --git a/Engine/Source/Editor/AnimGraph/Private/AnimGraphCommands.cpp b/Engine/Source/Editor/AnimGraph/Private/AnimGraphCommands.cpp
--- a/Engine/Source/Editor/AnimGraph/Private/AnimGraphCommands.cpp
+++ b/Engine/Source/Editor/AnimGraph/Private/AnimGraphCommands.cpp
@@ -6,7 +6,27 @@
 
 void FAnimGraphCommands::RegisterCommands()
 {
-	UI_COMMAND(TogglePoseWatch, "Toggle Pose Watch", "Toggle Pose Watching on this node", EUserInterfaceActionType::Button, FInputChord());
+	UI_COMMAND(TogglePoseWatch, "Toggle Pose Watch", "Toggle pose watching on this node", EUserInterfaceActionType::Button, FInputChord());
+
+	UI_COMMAND( SelectBone, "Select Bone", "Assign or change the bone for skeletal controls", EUserInterfaceActionType::Button, FInputChord() )
+	UI_COMMAND( AddBlendListPin, "Add Blend Pin", "Add blend pin to blend list", EUserInterfaceActionType::Button, FInputChord() )
+	UI_COMMAND( RemoveBlendListPin, "Remove Blend Pin", "Remove blend pin", EUserInterfaceActionType::Button, FInputChord() )
+
+	UI_COMMAND( ConvertToSeqEvaluator, "Convert to Single Frame Animation", "Convert to one frame animation that requires position", EUserInterfaceActionType::Button, FInputChord() )
+	UI_COMMAND( ConvertToSeqPlayer, "Convert to Sequence Player", "Convert back to sequence player without manual position set up", EUserInterfaceActionType::Button, FInputChord() )
+
+	UI_COMMAND( ConvertToBSEvaluator, "Convert to Single Frame Blend Space Player", "Convert to one frame blend space player that requires position", EUserInterfaceActionType::Button, FInputChord() )
+	UI_COMMAND( ConvertToBSPlayer, "Convert to Blend Space Player", "Convert back to blend space player without manual position set up", EUserInterfaceActionType::Button, FInputChord() )
+	UI_COMMAND( ConvertToBSGraph, "Convert to Blend Space Graph", "Convert to blend space graph", EUserInterfaceActionType::Button, FInputChord() )
+
+	UI_COMMAND( ConvertToAimOffsetLookAt, "Convert to LookAt AimOffset", "Convert to one AimOffset that automatically tracks a Target", EUserInterfaceActionType::Button, FInputChord())
+	UI_COMMAND( ConvertToAimOffsetSimple, "Convert to simple AimOffset", "Convert to a manual AimOffset", EUserInterfaceActionType::Button, FInputChord())
+	UI_COMMAND( ConvertToAimOffsetGraph, "Convert to AimOffset Graph", "Convert to an AimOffset graph", EUserInterfaceActionType::Button, FInputChord())
+
+	UI_COMMAND(ConvertToPoseBlender, "Convert to Pose Blender", "Convert to pose blender that can blend by source curves", EUserInterfaceActionType::Button, FInputChord())
+	UI_COMMAND(ConvertToPoseByName, "Convert to Pose By Name", "Convert to pose node that returns by name", EUserInterfaceActionType::Button, FInputChord())
+
+	UI_COMMAND( OpenRelatedAsset, "Open Asset", "Opens the asset related to this node", EUserInterfaceActionType::Button, FInputChord() )
 }
 
 #undef LOCTEXT_NAMESPACE
diff --git a/Engine/Source/Editor/AnimGraph/Private/AnimGraphNode_AimOffsetLookAt.cpp b/Engine/Source/Editor/AnimGraph/Private/AnimGraphNode_AimOffsetLookAt.cpp
--- a/Engine/Source/Editor/AnimGraph/Private/AnimGraphNode_AimOffsetLookAt.cpp
+++ b/Engine/Source/Editor/AnimGraph/Private/AnimGraphNode_AimOffsetLookAt.cpp
@@ -11,6 +11,7 @@
 #include "UObject/UObjectHash.h"
 #include "UObject/UObjectIterator.h"
 #include "ToolMenus.h"
+#include "AnimGraphCommands.h"
 
 /////////////////////////////////////////////////////
 // UAnimGraphNode_RotationOffsetBlendSpace
@@ -102,9 +103,12 @@ void UAnimGraphNode_AimOffsetLookAt::GetMenuActions(FBlueprintActionDatabaseRegi
 	{
 		if (const UBlendSpaceBase* TargetBlendSpace = Cast<UBlendSpaceBase>(RegistrarTarget))
 		{
-			if (UBlueprintNodeSpawner* NodeSpawner = GetMenuActions_Utils::MakeBlendSpaceAction(GetClass(), TargetBlendSpace))
+			if(TargetBlendSpace->IsAsset())
 			{
-				ActionRegistrar.AddBlueprintAction(TargetBlendSpace, NodeSpawner);
+				if (UBlueprintNodeSpawner* NodeSpawner = GetMenuActions_Utils::MakeBlendSpaceAction(GetClass(), TargetBlendSpace))
+				{
+					ActionRegistrar.AddBlueprintAction(TargetBlendSpace, NodeSpawner);
+				}
 			}
 		}
 		// else, the Blueprint database is specifically looking for actions pertaining to something different (not a BlendSpace asset)
@@ -115,9 +119,12 @@ void UAnimGraphNode_AimOffsetLookAt::GetMenuActions(FBlueprintActionDatabaseRegi
 		for (TObjectIterator<UBlendSpaceBase> BlendSpaceIt; BlendSpaceIt; ++BlendSpaceIt)
 		{
 			UBlendSpaceBase* BlendSpace = *BlendSpaceIt;
-			if (UBlueprintNodeSpawner* NodeSpawner = GetMenuActions_Utils::MakeBlendSpaceAction(NodeClass, BlendSpace))
+			if(BlendSpace->IsAsset())
 			{
-				ActionRegistrar.AddBlueprintAction(BlendSpace, NodeSpawner);
+				if (UBlueprintNodeSpawner* NodeSpawner = GetMenuActions_Utils::MakeBlendSpaceAction(NodeClass, BlendSpace))
+				{
+					ActionRegistrar.AddBlueprintAction(BlendSpace, NodeSpawner);
+				}
 			}
 		}
 	}
@@ -215,8 +222,8 @@ void UAnimGraphNode_AimOffsetLookAt::GetNodeContextMenuActions(UToolMenu* Menu,
 		// add an option to convert to single frame
 		{
 			FToolMenuSection& Section = Menu->AddSection("AnimGraphNodeBlendSpacePlayer", NSLOCTEXT("A3Nodes", "BlendSpaceHeading", "Blend Space"));
-			Section.AddMenuEntry(FGraphEditorCommands::Get().OpenRelatedAsset);
-			Section.AddMenuEntry(FGraphEditorCommands::Get().ConvertToAimOffsetSimple);
+			Section.AddMenuEntry(FAnimGraphCommands::Get().OpenRelatedAsset);
+			Section.AddMenuEntry(FAnimGraphCommands::Get().ConvertToAimOffsetSimple);
 		}
 	}
 }
diff --git a/Engine/Source/Editor/AnimGraph/Private/AnimGraphNode_Base.cpp b/Engine/Source/Editor/AnimGraph/Private/AnimGraphNode_Base.cpp
--- a/Engine/Source/Editor/AnimGraph/Private/AnimGraphNode_Base.cpp
+++ b/Engine/Source/Editor/AnimGraph/Private/AnimGraphNode_Base.cpp
@@ -133,6 +133,13 @@ void UAnimGraphNode_Base::ValidateAnimNodeDuringCompilation(USkeleton* ForSkelet
 	}
 }
 
+void UAnimGraphNode_Base::CopyTermDefaultsToDefaultObject(IAnimBlueprintCopyTermDefaultsContext& InCompilationContext, IAnimBlueprintNodeCopyTermDefaultsContext& InPerNodeContext, IAnimBlueprintGeneratedClassCompiledData& OutCompiledData)
+{
+	InPerNodeContext.GetTargetProperty()->CopyCompleteValue(InPerNodeContext.GetDestinationPtr(), InPerNodeContext.GetSourcePtr());
+
+	OnCopyTermDefaultsToDefaultObject(InCompilationContext, InPerNodeContext, OutCompiledData);
+}
+
 void UAnimGraphNode_Base::InternalPinCreation(TArray<UEdGraphPin*>* OldPins)
 {
 	// preload required assets first before creating pins
diff --git a/Engine/Source/Editor/AnimGraph/Private/AnimGraphNode_BlendListByInt.cpp b/Engine/Source/Editor/AnimGraph/Private/AnimGraphNode_BlendListByInt.cpp
--- a/Engine/Source/Editor/AnimGraph/Private/AnimGraphNode_BlendListByInt.cpp
+++ b/Engine/Source/Editor/AnimGraph/Private/AnimGraphNode_BlendListByInt.cpp
@@ -5,7 +5,7 @@
 #include "ToolMenus.h"
 #include "Kismet2/BlueprintEditorUtils.h"
 
-#include "GraphEditorActions.h"
+#include "AnimGraphCommands.h"
 #include "ScopedTransaction.h"
 
 /////////////////////////////////////////////////////
@@ -80,12 +80,12 @@ void UAnimGraphNode_BlendListByInt::GetNodeContextMenuActions(UToolMenu* Menu, U
 				// we only do this for normal BlendList/BlendList by enum, BlendList by Bool doesn't support add/remove pins
 				if ( Context->Pin->Direction == EGPD_Input )
 				{
-					Section.AddMenuEntry(FGraphEditorCommands::Get().RemoveBlendListPin);
+					Section.AddMenuEntry(FAnimGraphCommands::Get().RemoveBlendListPin);
 				}
 			}
 			else
 			{
-				Section.AddMenuEntry(FGraphEditorCommands::Get().AddBlendListPin);
+				Section.AddMenuEntry(FAnimGraphCommands::Get().AddBlendListPin);
 			}
 		}
 	}
diff --git a/Engine/Source/Editor/AnimGraph/Private/AnimGraphNode_BlendSpaceBase.cpp b/Engine/Source/Editor/AnimGraph/Private/AnimGraphNode_BlendSpaceBase.cpp
--- a/Engine/Source/Editor/AnimGraph/Private/AnimGraphNode_BlendSpaceBase.cpp
+++ b/Engine/Source/Editor/AnimGraph/Private/AnimGraphNode_BlendSpaceBase.cpp
@@ -10,39 +10,6 @@
 
 #define LOCTEXT_NAMESPACE "AnimGraphNode_BlendSpaceBase"
 
-/////////////////////////////////////////////////////
-
-// Action to add a sequence player node to the graph
-struct FNewBlendSpacePlayerAction : public FEdGraphSchemaAction_K2NewNode
-{
-	FNewBlendSpacePlayerAction(class UBlendSpaceBase* BlendSpace)
-		: FEdGraphSchemaAction_K2NewNode()
-	{
-		check(BlendSpace);
-
-		const bool bIsAimOffset = BlendSpace->IsA(UAimOffsetBlendSpace::StaticClass()) || BlendSpace->IsA(UAimOffsetBlendSpace1D::StaticClass());
-		FText NewTooltipDescription;
-		if (bIsAimOffset)
-		{
-			UAnimGraphNode_RotationOffsetBlendSpace* Template = NewObject<UAnimGraphNode_RotationOffsetBlendSpace>();
-			Template->Node.BlendSpace = BlendSpace;
-			NodeTemplate = Template;
-			NewTooltipDescription = LOCTEXT("EvalAimOffsetToMakePose", "Evaluates an aim offset at a particular coordinate to produce a pose");
-		}
-		else
-		{
-			UAnimGraphNode_BlendSpacePlayer* Template = NewObject<UAnimGraphNode_BlendSpacePlayer>();
-			Template->Node.BlendSpace = BlendSpace;
-			NodeTemplate = Template;
-			NewTooltipDescription = LOCTEXT("EvalBlendSpaceToMakePose", "Evaluates a blend space at a particular coordinate to produce a pose");
-		}
-
-		FText NewMenuDescription = NodeTemplate->GetNodeTitle(ENodeTitleType::ListView);
-
-		UpdateSearchData(NewMenuDescription, NewTooltipDescription, LOCTEXT("Animation", "Animations"), FText::FromString(BlendSpace->GetPathName()));
-	}
-};
-
 /////////////////////////////////////////////////////
 // UAnimGraphNode_BlendSpaceBase
 
diff --git a/Engine/Source/Editor/AnimGraph/Private/AnimGraphNode_BlendSpaceEvaluator.cpp b/Engine/Source/Editor/AnimGraph/Private/AnimGraphNode_BlendSpaceEvaluator.cpp
--- a/Engine/Source/Editor/AnimGraph/Private/AnimGraphNode_BlendSpaceEvaluator.cpp
+++ b/Engine/Source/Editor/AnimGraph/Private/AnimGraphNode_BlendSpaceEvaluator.cpp
@@ -2,7 +2,7 @@
 
 #include "AnimGraphNode_BlendSpaceEvaluator.h"
 #include "ToolMenus.h"
-#include "GraphEditorActions.h"
+#include "AnimGraphCommands.h"
 #include "Kismet2/CompilerResultsLog.h"
 
 /////////////////////////////////////////////////////
@@ -149,8 +149,8 @@ void UAnimGraphNode_BlendSpaceEvaluator::GetNodeContextMenuActions(UToolMenu* Me
 		// add an option to convert to single frame
 		{
 			FToolMenuSection& Section = Menu->AddSection("AnimGraphNodeBlendSpacePlayer", NSLOCTEXT("A3Nodes", "BlendSpaceHeading", "Blend Space"));
-			Section.AddMenuEntry(FGraphEditorCommands::Get().OpenRelatedAsset);
-			Section.AddMenuEntry(FGraphEditorCommands::Get().ConvertToBSPlayer);
+			Section.AddMenuEntry(FAnimGraphCommands::Get().OpenRelatedAsset);
+			Section.AddMenuEntry(FAnimGraphCommands::Get().ConvertToBSPlayer);
 		}
 	}
 }
diff --git a/Engine/Source/Editor/AnimGraph/Private/AnimGraphNode_BlendSpaceGraph.cpp b/Engine/Source/Editor/AnimGraph/Private/AnimGraphNode_BlendSpaceGraph.cpp
new file mode 100644
--- /dev/null
+++ b/Engine/Source/Editor/AnimGraph/Private/AnimGraphNode_BlendSpaceGraph.cpp
@@ -0,0 +1,138 @@
+// Copyright Epic Games, Inc. All Rights Reserved.
+
+#include "AnimGraphNode_BlendSpaceGraph.h"
+#include "BlueprintNodeSpawner.h"
+#include "BlueprintActionDatabaseRegistrar.h"
+#include "Animation/AimOffsetBlendSpace.h"
+#include "Animation/AimOffsetBlendSpace1D.h"
+
+#define LOCTEXT_NAMESPACE "UAnimGraphNode_BlendSpaceGraph"
+
+FText UAnimGraphNode_BlendSpaceGraph::GetNodeTitle(ENodeTitleType::Type TitleType) const
+{
+	if(BlendSpaceGraph || BlendSpace)
+	{
+		const FText BlendSpaceName = FText::FromString(BlendSpaceGraph ? GetBlendSpaceGraphName() : GetBlendSpaceName());
+
+		if(TitleType == ENodeTitleType::EditableTitle)
+		{
+			return BlendSpaceName;
+		}
+		else if (TitleType == ENodeTitleType::ListView || TitleType == ENodeTitleType::MenuTitle)
+		{
+			FFormatNamedArguments Args;
+			Args.Add(TEXT("BlendSpaceName"), BlendSpaceName);
+			return FText::Format(LOCTEXT("BlendspaceListTitle", "Blendspace '{BlendSpaceName}'"), Args);
+		}
+		else
+		{
+			FFormatNamedArguments TitleArgs;
+			TitleArgs.Add(TEXT("BlendSpaceName"), BlendSpaceName);
+			FText Title = FText::Format(LOCTEXT("BlendSpaceFullTitle", "{BlendSpaceName}\nBlendspace"), TitleArgs);
+
+			if ((TitleType == ENodeTitleType::FullTitle) && (Node.GetGroupName() != NAME_None))
+			{
+				FFormatNamedArguments Args;
+				Args.Add(TEXT("Title"), Title);
+				Args.Add(TEXT("SyncGroupName"), FText::FromName(Node.GetGroupName()));
+				Title = FText::Format(LOCTEXT("BlendSpaceNodeGroupSubtitle", "{Title}\nSync group {SyncGroupName}"), Args);
+			}
+
+			return Title;
+		}
+	}
+	else if(BlendSpaceClass.Get())
+	{
+		return BlendSpaceClass.Get()->GetDisplayNameText();
+	}
+	else
+	{
+		return LOCTEXT("EmptyBlendspaceListTitle", "Blendspace");
+	}
+}
+
+void UAnimGraphNode_BlendSpaceGraph::GetMenuActions(FBlueprintActionDatabaseRegistrar& ActionRegistrar) const
+{
+	struct GetMenuActions_Utils
+	{
+		static void SetNodeBlendSpace(UEdGraphNode* NewNode, bool bIsTemplateNode, TWeakObjectPtr<UBlendSpaceBase> BlendSpace)
+		{
+			UAnimGraphNode_BlendSpaceGraph* BlendSpaceNode = CastChecked<UAnimGraphNode_BlendSpaceGraph>(NewNode);
+			BlendSpaceNode->SetupFromAsset(BlendSpace.Get(), bIsTemplateNode);
+		}
+
+		static UBlueprintNodeSpawner* MakeBlendSpaceAction(TSubclassOf<UEdGraphNode> const NodeClass, const UBlendSpaceBase* InBlendSpace)
+		{
+			UBlueprintNodeSpawner* NodeSpawner = nullptr;
+
+			bool const bIsAimOffset = InBlendSpace->IsA(UAimOffsetBlendSpace::StaticClass()) ||
+									  InBlendSpace->IsA(UAimOffsetBlendSpace1D::StaticClass());
+			if (!bIsAimOffset)
+			{
+				NodeSpawner = UBlueprintNodeSpawner::Create(NodeClass);
+				check(NodeSpawner != nullptr);
+
+				TWeakObjectPtr<UBlendSpaceBase> BlendSpacePtr = MakeWeakObjectPtr(const_cast<UBlendSpaceBase*>(InBlendSpace));
+				NodeSpawner->CustomizeNodeDelegate = UBlueprintNodeSpawner::FCustomizeNodeDelegate::CreateStatic(GetMenuActions_Utils::SetNodeBlendSpace, BlendSpacePtr);
+			}	
+			return NodeSpawner;
+		}
+
+		static void SetNodeBlendSpaceClass(UEdGraphNode* NewNode, bool bIsTemplateNode, TSubclassOf<UBlendSpaceBase> InBlendSpaceClass)
+		{
+			UAnimGraphNode_BlendSpaceGraph* BlendSpaceNode = CastChecked<UAnimGraphNode_BlendSpaceGraph>(NewNode);
+			BlendSpaceNode->SetupFromClass(InBlendSpaceClass, bIsTemplateNode);
+		}
+
+		static UBlueprintNodeSpawner* MakeBlendSpaceAction(TSubclassOf<UEdGraphNode> const NodeClass, TSubclassOf<UBlendSpaceBase> InBlendSpaceClass)
+		{
+			UBlueprintNodeSpawner* NodeSpawner = UBlueprintNodeSpawner::Create(NodeClass);
+			check(NodeSpawner != nullptr);
+
+			NodeSpawner->CustomizeNodeDelegate = UBlueprintNodeSpawner::FCustomizeNodeDelegate::CreateStatic(GetMenuActions_Utils::SetNodeBlendSpaceClass, InBlendSpaceClass);
+	
+			return NodeSpawner;
+		}
+	};
+
+	// Add the non-asset based cases
+	ActionRegistrar.AddBlueprintAction(GetMenuActions_Utils::MakeBlendSpaceAction(GetClass(), UBlendSpace::StaticClass()));
+	ActionRegistrar.AddBlueprintAction(GetMenuActions_Utils::MakeBlendSpaceAction(GetClass(), UBlendSpace1D::StaticClass()));
+
+	if (const UObject* RegistrarTarget = ActionRegistrar.GetActionKeyFilter())
+	{
+		if (const UBlendSpaceBase* TargetBlendSpace = Cast<UBlendSpaceBase>(RegistrarTarget))
+		{
+			if(TargetBlendSpace->IsAsset())
+			{
+				if (UBlueprintNodeSpawner* NodeSpawner = GetMenuActions_Utils::MakeBlendSpaceAction(GetClass(), TargetBlendSpace))
+				{
+					ActionRegistrar.AddBlueprintAction(TargetBlendSpace, NodeSpawner);
+				}
+			}
+		}
+		// else, the Blueprint database is specifically looking for actions pertaining to something different (not a BlendSpace asset)
+	}
+	else
+	{
+		UClass* NodeClass = GetClass();
+		for (TObjectIterator<UBlendSpaceBase> BlendSpaceIt; BlendSpaceIt; ++BlendSpaceIt)
+		{
+			if(BlendSpaceIt->IsAsset())
+			{
+				if (UBlueprintNodeSpawner* NodeSpawner = GetMenuActions_Utils::MakeBlendSpaceAction(NodeClass, *BlendSpaceIt))
+				{
+					ActionRegistrar.AddBlueprintAction(*BlendSpaceIt, NodeSpawner);
+				}
+			}
+		}
+	}
+}
+
+void UAnimGraphNode_BlendSpaceGraph::BakeDataDuringCompilation(class FCompilerResultsLog& MessageLog)
+{
+	UAnimBlueprint* AnimBlueprint = GetAnimBlueprint();
+	AnimBlueprint->FindOrAddGroup(Node.GetGroupName());
+}
+
+#undef LOCTEXT_NAMESPACE
diff --git a/Engine/Source/Editor/AnimGraph/Private/AnimGraphNode_BlendSpacePlayer.cpp b/Engine/Source/Editor/AnimGraph/Private/AnimGraphNode_BlendSpacePlayer.cpp
--- a/Engine/Source/Editor/AnimGraph/Private/AnimGraphNode_BlendSpacePlayer.cpp
+++ b/Engine/Source/Editor/AnimGraph/Private/AnimGraphNode_BlendSpacePlayer.cpp
@@ -10,6 +10,7 @@
 #include "BlueprintActionDatabaseRegistrar.h"
 #include "Animation/AimOffsetBlendSpace.h"
 #include "Animation/AimOffsetBlendSpace1D.h"
+#include "AnimGraphCommands.h"
 
 /////////////////////////////////////////////////////
 // UAnimGraphNode_BlendSpacePlayer
@@ -171,8 +172,9 @@ void UAnimGraphNode_BlendSpacePlayer::GetNodeContextMenuActions(UToolMenu* Menu,
 		// add an option to convert to single frame
 		{
 			FToolMenuSection& Section = Menu->AddSection("AnimGraphNodeBlendSpaceEvaluator", NSLOCTEXT("A3Nodes", "BlendSpaceHeading", "Blend Space"));
-			Section.AddMenuEntry(FGraphEditorCommands::Get().OpenRelatedAsset);
-			Section.AddMenuEntry(FGraphEditorCommands::Get().ConvertToBSEvaluator);
+			Section.AddMenuEntry(FAnimGraphCommands::Get().OpenRelatedAsset);
+			Section.AddMenuEntry(FAnimGraphCommands::Get().ConvertToBSEvaluator);
+			Section.AddMenuEntry(FAnimGraphCommands::Get().ConvertToBSGraph);
 		}
 	}
 }
@@ -209,9 +211,12 @@ void UAnimGraphNode_BlendSpacePlayer::GetMenuActions(FBlueprintActionDatabaseReg
 	{
 		if (const UBlendSpaceBase* TargetBlendSpace = Cast<UBlendSpaceBase>(RegistrarTarget))
 		{
-			if (UBlueprintNodeSpawner* NodeSpawner = GetMenuActions_Utils::MakeBlendSpaceAction(GetClass(), TargetBlendSpace))
+			if(TargetBlendSpace->IsAsset())
 			{
-				ActionRegistrar.AddBlueprintAction(TargetBlendSpace, NodeSpawner);
+				if (UBlueprintNodeSpawner* NodeSpawner = GetMenuActions_Utils::MakeBlendSpaceAction(GetClass(), TargetBlendSpace))
+				{
+					ActionRegistrar.AddBlueprintAction(TargetBlendSpace, NodeSpawner);
+				}
 			}
 		}
 		// else, the Blueprint database is specifically looking for actions pertaining to something different (not a BlendSpace asset)
@@ -222,9 +227,12 @@ void UAnimGraphNode_BlendSpacePlayer::GetMenuActions(FBlueprintActionDatabaseReg
 		for (TObjectIterator<UBlendSpaceBase> BlendSpaceIt; BlendSpaceIt; ++BlendSpaceIt)
 		{
 			UBlendSpaceBase* BlendSpace = *BlendSpaceIt;
-			if (UBlueprintNodeSpawner* NodeSpawner = GetMenuActions_Utils::MakeBlendSpaceAction(NodeClass, BlendSpace))
+			if(BlendSpace->IsAsset())
 			{
-				ActionRegistrar.AddBlueprintAction(BlendSpace, NodeSpawner);
+				if (UBlueprintNodeSpawner* NodeSpawner = GetMenuActions_Utils::MakeBlendSpaceAction(NodeClass, BlendSpace))
+				{
+					ActionRegistrar.AddBlueprintAction(BlendSpace, NodeSpawner);
+				}
 			}
 		}
 	}
diff --git a/Engine/Source/Editor/AnimGraph/Private/AnimGraphNode_LayeredBoneBlend.cpp b/Engine/Source/Editor/AnimGraph/Private/AnimGraphNode_LayeredBoneBlend.cpp
--- a/Engine/Source/Editor/AnimGraph/Private/AnimGraphNode_LayeredBoneBlend.cpp
+++ b/Engine/Source/Editor/AnimGraph/Private/AnimGraphNode_LayeredBoneBlend.cpp
@@ -4,7 +4,7 @@
 #include "ToolMenus.h"
 #include "Kismet2/BlueprintEditorUtils.h"
 
-#include "GraphEditorActions.h"
+#include "AnimGraphCommands.h"
 #include "ScopedTransaction.h"
 
 #include "DetailLayoutBuilder.h"
@@ -130,12 +130,12 @@ void UAnimGraphNode_LayeredBoneBlend::GetNodeContextMenuActions(UToolMenu* Menu,
 				if (Context->Pin->Direction == EGPD_Input)
 				{
 					//@TODO: Only offer this option on arrayed pins
-					Section.AddMenuEntry(FGraphEditorCommands::Get().RemoveBlendListPin);
+					Section.AddMenuEntry(FAnimGraphCommands::Get().RemoveBlendListPin);
 				}
 			}
 			else
 			{
-				Section.AddMenuEntry(FGraphEditorCommands::Get().AddBlendListPin);
+				Section.AddMenuEntry(FAnimGraphCommands::Get().AddBlendListPin);
 			}
 		}
 	}
diff --git a/Engine/Source/Editor/AnimGraph/Private/AnimGraphNode_LinkedAnimGraphBase.cpp b/Engine/Source/Editor/AnimGraph/Private/AnimGraphNode_LinkedAnimGraphBase.cpp
--- a/Engine/Source/Editor/AnimGraph/Private/AnimGraphNode_LinkedAnimGraphBase.cpp
+++ b/Engine/Source/Editor/AnimGraph/Private/AnimGraphNode_LinkedAnimGraphBase.cpp
@@ -17,6 +17,7 @@
 #include "UObject/CoreRedirects.h"
 #include "Animation/AnimNode_LinkedAnimGraph.h"
 #include "AnimGraphAttributes.h"
+#include "IAnimBlueprintCopyTermDefaultsContext.h"
 
 #define LOCTEXT_NAMESPACE "LinkedAnimGraph"
 
@@ -481,4 +482,12 @@ void UAnimGraphNode_LinkedAnimGraphBase::GetOutputLinkAttributes(FNodeAttributeA
 	});
 }
 
+void UAnimGraphNode_LinkedAnimGraphBase::OnCopyTermDefaultsToDefaultObject(IAnimBlueprintCopyTermDefaultsContext& InCompilationContext, IAnimBlueprintNodeCopyTermDefaultsContext& InPerNodeContext, IAnimBlueprintGeneratedClassCompiledData& OutCompiledData)
+{
+	UAnimGraphNode_LinkedAnimGraphBase* TrueNode = InCompilationContext.GetMessageLog().FindSourceObjectTypeChecked<UAnimGraphNode_LinkedAnimGraphBase>(this);
+
+	FAnimNode_LinkedAnimGraph* DestinationNode = reinterpret_cast<FAnimNode_LinkedAnimGraph*>(InPerNodeContext.GetDestinationPtr());
+	DestinationNode->NodeIndex = InPerNodeContext.GetNodePropertyIndex();
+}
+
 #undef LOCTEXT_NAMESPACE
diff --git a/Engine/Source/Editor/AnimGraph/Private/AnimGraphNode_LinkedInputPose.cpp b/Engine/Source/Editor/AnimGraph/Private/AnimGraphNode_LinkedInputPose.cpp
--- a/Engine/Source/Editor/AnimGraph/Private/AnimGraphNode_LinkedInputPose.cpp
+++ b/Engine/Source/Editor/AnimGraph/Private/AnimGraphNode_LinkedInputPose.cpp
@@ -28,6 +28,7 @@
 #include "K2Node_VariableGet.h"
 #include "AnimBlueprintCompiler.h"
 #include "AnimGraphAttributes.h"
+#include "IAnimBlueprintCopyTermDefaultsContext.h"
 
 #define LOCTEXT_NAMESPACE "LinkedInputPose"
 
@@ -502,6 +503,14 @@ void UAnimGraphNode_LinkedInputPose::CustomizeDetails(IDetailLayoutBuilder& Deta
 		.ShouldAutoExpand(true);
 }
 
+void UAnimGraphNode_LinkedInputPose::OnCopyTermDefaultsToDefaultObject(IAnimBlueprintCopyTermDefaultsContext& InCompilationContext, IAnimBlueprintNodeCopyTermDefaultsContext& InPerNodeContext, IAnimBlueprintGeneratedClassCompiledData& OutCompiledData)
+{
+	UAnimGraphNode_LinkedInputPose* TrueNode = InCompilationContext.GetMessageLog().FindSourceObjectTypeChecked<UAnimGraphNode_LinkedInputPose>(this);
+
+	FAnimNode_LinkedInputPose* DestinationNode = reinterpret_cast<FAnimNode_LinkedInputPose*>(InPerNodeContext.GetDestinationPtr());
+	DestinationNode->Graph = TrueNode->GetGraph()->GetFName();
+}
+
 TSharedRef<SWidget> UAnimGraphNode_LinkedInputPose::MakeNameWidget(IDetailLayoutBuilder& DetailBuilder)
 {
 	TSharedPtr<IPropertyHandle> NamePropertyHandle = DetailBuilder.GetProperty(GET_MEMBER_NAME_CHECKED(UAnimGraphNode_LinkedInputPose, Node.Name), GetClass());
diff --git a/Engine/Source/Editor/AnimGraph/Private/AnimGraphNode_PoseBlendNode.cpp b/Engine/Source/Editor/AnimGraph/Private/AnimGraphNode_PoseBlendNode.cpp
--- a/Engine/Source/Editor/AnimGraph/Private/AnimGraphNode_PoseBlendNode.cpp
+++ b/Engine/Source/Editor/AnimGraph/Private/AnimGraphNode_PoseBlendNode.cpp
@@ -5,7 +5,7 @@
 #include "Modules/ModuleManager.h"
 #include "ToolMenus.h"
 
-#include "GraphEditorActions.h"
+#include "AnimGraphCommands.h"
 #include "ARFilter.h"
 #include "AssetRegistryModule.h"
 #include "BlueprintActionFilter.h"
@@ -251,7 +251,7 @@ void UAnimGraphNode_PoseBlendNode::GetNodeContextMenuActions(UToolMenu* Menu, UG
 		// add an option to convert to single frame
 		{
 			FToolMenuSection& Section = Menu->AddSection("AnimGraphNodePoseBlender", LOCTEXT("PoseBlenderHeading", "Pose Blender"));
-			Section.AddMenuEntry(FGraphEditorCommands::Get().ConvertToPoseByName);
+			Section.AddMenuEntry(FAnimGraphCommands::Get().ConvertToPoseByName);
 		}
 	}
 }
diff --git a/Engine/Source/Editor/AnimGraph/Private/AnimGraphNode_PoseByName.cpp b/Engine/Source/Editor/AnimGraph/Private/AnimGraphNode_PoseByName.cpp
--- a/Engine/Source/Editor/AnimGraph/Private/AnimGraphNode_PoseByName.cpp
+++ b/Engine/Source/Editor/AnimGraph/Private/AnimGraphNode_PoseByName.cpp
@@ -4,7 +4,7 @@
 #include "ToolMenus.h"
 
 #include "Kismet2/CompilerResultsLog.h"
-#include "GraphEditorActions.h"
+#include "AnimGraphCommands.h"
 
 /////////////////////////////////////////////////////
 // UAnimGraphNode_PoseByName
@@ -155,7 +155,7 @@ void UAnimGraphNode_PoseByName::GetNodeContextMenuActions(UToolMenu* Menu, UGrap
 		// add an option to convert to single frame
 		{
 			FToolMenuSection& Section = Menu->AddSection("AnimGraphNodePoseByName", LOCTEXT("PoseByNameHeading", "Pose By Name"));
-			Section.AddMenuEntry(FGraphEditorCommands::Get().ConvertToPoseBlender);
+			Section.AddMenuEntry(FAnimGraphCommands::Get().ConvertToPoseBlender);
 		}
 	}
 }
diff --git a/Engine/Source/Editor/AnimGraph/Private/AnimGraphNode_Root.cpp b/Engine/Source/Editor/AnimGraph/Private/AnimGraphNode_Root.cpp
--- a/Engine/Source/Editor/AnimGraph/Private/AnimGraphNode_Root.cpp
+++ b/Engine/Source/Editor/AnimGraph/Private/AnimGraphNode_Root.cpp
@@ -5,6 +5,7 @@
 #include "AnimBlueprintCompiler.h"
 #include "AnimBlueprintCompilerHandler_Base.h"
 #include "IAnimBlueprintCompilationContext.h"
+#include "IAnimBlueprintCopyTermDefaultsContext.h"
 
 /////////////////////////////////////////////////////
 // FPoseLinkMappingRecord
@@ -94,4 +95,12 @@ void UAnimGraphNode_Root::OnProcessDuringCompilation(IAnimBlueprintCompilationCo
 	Node.Name = TrueNode->GetGraph()->GetFName();
 }
 
+void UAnimGraphNode_Root::OnCopyTermDefaultsToDefaultObject(IAnimBlueprintCopyTermDefaultsContext& InCompilationContext, IAnimBlueprintNodeCopyTermDefaultsContext& InPerNodeContext, IAnimBlueprintGeneratedClassCompiledData& OutCompiledData)
+{
+	UAnimGraphNode_Root* TrueNode = InCompilationContext.GetMessageLog().FindSourceObjectTypeChecked<UAnimGraphNode_Root>(this);
+
+	FAnimNode_Root* DestinationNode = reinterpret_cast<FAnimNode_Root*>(InPerNodeContext.GetDestinationPtr());
+	DestinationNode->Name = TrueNode->GetGraph()->GetFName();
+}
+
 #undef LOCTEXT_NAMESPACE
diff --git a/Engine/Source/Editor/AnimGraph/Private/AnimGraphNode_RotationOffsetBlendSpace.cpp b/Engine/Source/Editor/AnimGraph/Private/AnimGraphNode_RotationOffsetBlendSpace.cpp
--- a/Engine/Source/Editor/AnimGraph/Private/AnimGraphNode_RotationOffsetBlendSpace.cpp
+++ b/Engine/Source/Editor/AnimGraph/Private/AnimGraphNode_RotationOffsetBlendSpace.cpp
@@ -14,6 +14,7 @@
 #include "DetailLayoutBuilder.h"
 #include "ScopedTransaction.h"
 #include "Kismet2/BlueprintEditorUtils.h"
+#include "AnimGraphCommands.h"
 
 /////////////////////////////////////////////////////
 // UAnimGraphNode_RotationOffsetBlendSpace
@@ -105,9 +106,12 @@ void UAnimGraphNode_RotationOffsetBlendSpace::GetMenuActions(FBlueprintActionDat
 	{
 		if (const UBlendSpaceBase* TargetBlendSpace = Cast<UBlendSpaceBase>(RegistrarTarget))
 		{
-			if (UBlueprintNodeSpawner* NodeSpawner = GetMenuActions_Utils::MakeBlendSpaceAction(GetClass(), TargetBlendSpace))
+			if(TargetBlendSpace->IsAsset())
 			{
-				ActionRegistrar.AddBlueprintAction(TargetBlendSpace, NodeSpawner);
+				if (UBlueprintNodeSpawner* NodeSpawner = GetMenuActions_Utils::MakeBlendSpaceAction(GetClass(), TargetBlendSpace))
+				{
+					ActionRegistrar.AddBlueprintAction(TargetBlendSpace, NodeSpawner);
+				}
 			}
 		}
 		// else, the Blueprint database is specifically looking for actions pertaining to something different (not a BlendSpace asset)
@@ -118,9 +122,12 @@ void UAnimGraphNode_RotationOffsetBlendSpace::GetMenuActions(FBlueprintActionDat
 		for (TObjectIterator<UBlendSpaceBase> BlendSpaceIt; BlendSpaceIt; ++BlendSpaceIt)
 		{
 			UBlendSpaceBase* BlendSpace = *BlendSpaceIt;
-			if (UBlueprintNodeSpawner* NodeSpawner = GetMenuActions_Utils::MakeBlendSpaceAction(NodeClass, BlendSpace))
+			if(BlendSpace->IsAsset())
 			{
-				ActionRegistrar.AddBlueprintAction(BlendSpace, NodeSpawner);
+				if (UBlueprintNodeSpawner* NodeSpawner = GetMenuActions_Utils::MakeBlendSpaceAction(NodeClass, BlendSpace))
+				{
+					ActionRegistrar.AddBlueprintAction(BlendSpace, NodeSpawner);
+				}
 			}
 		}
 	}
@@ -192,8 +199,9 @@ void UAnimGraphNode_RotationOffsetBlendSpace::GetNodeContextMenuActions(UToolMen
 		// add an option to convert to single frame
 		{
 			FToolMenuSection& Section = Menu->AddSection("AnimGraphNodeBlendSpacePlayer", NSLOCTEXT("A3Nodes", "BlendSpaceHeading", "Blend Space"));
-			Section.AddMenuEntry(FGraphEditorCommands::Get().OpenRelatedAsset);
-			Section.AddMenuEntry(FGraphEditorCommands::Get().ConvertToAimOffsetLookAt);
+			Section.AddMenuEntry(FAnimGraphCommands::Get().OpenRelatedAsset);
+			Section.AddMenuEntry(FAnimGraphCommands::Get().ConvertToAimOffsetLookAt);
+			Section.AddMenuEntry(FAnimGraphCommands::Get().ConvertToAimOffsetGraph);
 		}
 	}
 }
diff --git a/Engine/Source/Editor/AnimGraph/Private/AnimGraphNode_RotationOffsetBlendSpaceGraph.cpp b/Engine/Source/Editor/AnimGraph/Private/AnimGraphNode_RotationOffsetBlendSpaceGraph.cpp
new file mode 100644
--- /dev/null
+++ b/Engine/Source/Editor/AnimGraph/Private/AnimGraphNode_RotationOffsetBlendSpaceGraph.cpp
@@ -0,0 +1,163 @@
+// Copyright Epic Games, Inc. All Rights Reserved.
+
+#include "AnimGraphNode_RotationOffsetBlendSpaceGraph.h"
+#include "BlueprintNodeSpawner.h"
+#include "BlueprintActionDatabaseRegistrar.h"
+#include "Animation/AimOffsetBlendSpace.h"
+#include "Animation/AimOffsetBlendSpace1D.h"
+#include "AnimGraphNodeAlphaOptions.h"
+#include "DetailLayoutBuilder.h"
+
+#define LOCTEXT_NAMESPACE "UAnimGraphNode_RotationOffsetBlendSpaceGraph"
+
+FText UAnimGraphNode_RotationOffsetBlendSpaceGraph::GetNodeTitle(ENodeTitleType::Type TitleType) const
+{
+	if(BlendSpaceGraph || BlendSpace)
+	{
+		const FText BlendSpaceName = FText::FromString(BlendSpaceGraph ? GetBlendSpaceGraphName() : GetBlendSpaceName());
+
+		if(TitleType == ENodeTitleType::EditableTitle)
+		{
+			return BlendSpaceName;
+		}
+		else if (TitleType == ENodeTitleType::ListView || TitleType == ENodeTitleType::MenuTitle)
+		{
+			FFormatNamedArguments Args;
+			Args.Add(TEXT("BlendSpaceName"), BlendSpaceName);
+			return FText::Format(LOCTEXT("AimOffsetListTitle", "AimOffset '{BlendSpaceName}'"), Args);
+		}
+		else
+		{
+			FFormatNamedArguments TitleArgs;
+			TitleArgs.Add(TEXT("BlendSpaceName"), BlendSpaceName);
+			FText Title = FText::Format(LOCTEXT("AimOffsetFullTitle", "{BlendSpaceName}\nAimOffset"), TitleArgs);
+
+			if ((TitleType == ENodeTitleType::FullTitle) && (Node.GetGroupName() != NAME_None))
+			{
+				FFormatNamedArguments Args;
+				Args.Add(TEXT("Title"), Title);
+				Args.Add(TEXT("SyncGroupName"), FText::FromName(Node.GetGroupName()));
+				Title = FText::Format(LOCTEXT("AimOffsetNodeGroupSubtitle", "{Title}\nSync group {SyncGroupName}"), Args);
+			}
+
+			return Title;
+		}
+	}
+	else if(BlendSpaceClass.Get())
+	{
+		return BlendSpaceClass.Get()->GetDisplayNameText();
+	}
+	else
+	{
+		return LOCTEXT("EmptyBlendspaceListTitle", "AimOffset");
+	}
+}
+
+void UAnimGraphNode_RotationOffsetBlendSpaceGraph::GetMenuActions(FBlueprintActionDatabaseRegistrar& ActionRegistrar) const
+{
+	struct GetMenuActions_Utils
+	{
+		static void SetNodeBlendSpace(UEdGraphNode* NewNode, bool bIsTemplateNode, TWeakObjectPtr<UBlendSpaceBase> BlendSpace)
+		{
+			UAnimGraphNode_RotationOffsetBlendSpaceGraph* BlendSpaceNode = CastChecked<UAnimGraphNode_RotationOffsetBlendSpaceGraph>(NewNode);
+			BlendSpaceNode->SetupFromAsset(BlendSpace.Get(), bIsTemplateNode);
+		}
+
+		static UBlueprintNodeSpawner* MakeBlendSpaceAction(TSubclassOf<UEdGraphNode> const NodeClass, const UBlendSpaceBase* InBlendSpace)
+		{
+			UBlueprintNodeSpawner* NodeSpawner = nullptr;
+
+			bool const bIsAimOffset = InBlendSpace->IsA(UAimOffsetBlendSpace::StaticClass()) ||
+									  InBlendSpace->IsA(UAimOffsetBlendSpace1D::StaticClass());
+			if (bIsAimOffset)
+			{
+				NodeSpawner = UBlueprintNodeSpawner::Create(NodeClass);
+				check(NodeSpawner != nullptr);
+
+				TWeakObjectPtr<UBlendSpaceBase> BlendSpacePtr = MakeWeakObjectPtr(const_cast<UBlendSpaceBase*>(InBlendSpace));
+				NodeSpawner->CustomizeNodeDelegate = UBlueprintNodeSpawner::FCustomizeNodeDelegate::CreateStatic(GetMenuActions_Utils::SetNodeBlendSpace, BlendSpacePtr);
+			}	
+			return NodeSpawner;
+		}
+
+		static void SetNodeBlendSpaceClass(UEdGraphNode* NewNode, bool bIsTemplateNode, TSubclassOf<UBlendSpaceBase> InBlendSpaceClass)
+		{
+			UAnimGraphNode_RotationOffsetBlendSpaceGraph* BlendSpaceNode = CastChecked<UAnimGraphNode_RotationOffsetBlendSpaceGraph>(NewNode);
+			BlendSpaceNode->SetupFromClass(InBlendSpaceClass, bIsTemplateNode);
+		}
+
+		static UBlueprintNodeSpawner* MakeBlendSpaceAction(TSubclassOf<UEdGraphNode> const NodeClass, TSubclassOf<UBlendSpaceBase> InBlendSpaceClass)
+		{
+			UBlueprintNodeSpawner* NodeSpawner = UBlueprintNodeSpawner::Create(NodeClass);
+			check(NodeSpawner != nullptr);
+
+			NodeSpawner->CustomizeNodeDelegate = UBlueprintNodeSpawner::FCustomizeNodeDelegate::CreateStatic(GetMenuActions_Utils::SetNodeBlendSpaceClass, InBlendSpaceClass);
+	
+			return NodeSpawner;
+		}
+	};
+
+	// Add the non-asset based cases
+	ActionRegistrar.AddBlueprintAction(GetMenuActions_Utils::MakeBlendSpaceAction(GetClass(), UAimOffsetBlendSpace::StaticClass()));
+	ActionRegistrar.AddBlueprintAction(GetMenuActions_Utils::MakeBlendSpaceAction(GetClass(), UAimOffsetBlendSpace1D::StaticClass()));
+
+	if (const UObject* RegistrarTarget = ActionRegistrar.GetActionKeyFilter())
+	{
+		if (const UBlendSpaceBase* TargetBlendSpace = Cast<UBlendSpaceBase>(RegistrarTarget))
+		{
+			if(TargetBlendSpace->IsAsset())
+			{
+				if (UBlueprintNodeSpawner* NodeSpawner = GetMenuActions_Utils::MakeBlendSpaceAction(GetClass(), TargetBlendSpace))
+				{
+					ActionRegistrar.AddBlueprintAction(TargetBlendSpace, NodeSpawner);
+				}
+			}
+		}
+		// else, the Blueprint database is specifically looking for actions pertaining to something different (not a BlendSpace asset)
+	}
+	else
+	{
+		UClass* NodeClass = GetClass();
+		for (TObjectIterator<UBlendSpaceBase> BlendSpaceIt; BlendSpaceIt; ++BlendSpaceIt)
+		{
+			if(BlendSpaceIt->IsAsset())
+			{
+				if (UBlueprintNodeSpawner* NodeSpawner = GetMenuActions_Utils::MakeBlendSpaceAction(NodeClass, *BlendSpaceIt))
+				{
+					ActionRegistrar.AddBlueprintAction(*BlendSpaceIt, NodeSpawner);
+				}
+			}
+		}
+	}
+}
+
+void UAnimGraphNode_RotationOffsetBlendSpaceGraph::BakeDataDuringCompilation(FCompilerResultsLog& MessageLog)
+{
+	UAnimBlueprint* AnimBlueprint = GetAnimBlueprint();
+	AnimBlueprint->FindOrAddGroup(Node.GetGroupName());
+}
+
+void UAnimGraphNode_RotationOffsetBlendSpaceGraph::CustomizePinData(UEdGraphPin* Pin, FName SourcePropertyName, int32 ArrayIndex) const
+{
+	Super::CustomizePinData(Pin, SourcePropertyName, ArrayIndex);
+
+	FAnimGraphNodeAlphaOptions::HandleCustomizePinData(Node, Pin);
+}
+
+void UAnimGraphNode_RotationOffsetBlendSpaceGraph::PostEditChangeProperty(FPropertyChangedEvent& PropertyChangedEvent)
+{
+	FAnimGraphNodeAlphaOptions::HandlePostEditChangeProperty(Node, this, PropertyChangedEvent);
+
+	Super::PostEditChangeProperty(PropertyChangedEvent);
+}
+
+void UAnimGraphNode_RotationOffsetBlendSpaceGraph::CustomizeDetails(IDetailLayoutBuilder& InDetailBuilder)
+{
+	Super::CustomizeDetails(InDetailBuilder);
+
+	TSharedRef<IPropertyHandle> NodeHandle = InDetailBuilder.GetProperty(TEXT("Node"), GetClass());
+
+	FAnimGraphNodeAlphaOptions::HandleCustomizeDetails(Node, NodeHandle, InDetailBuilder);
+}
+
+#undef LOCTEXT_NAMESPACE
diff --git a/Engine/Source/Editor/AnimGraph/Private/AnimGraphNode_SequenceEvaluator.cpp b/Engine/Source/Editor/AnimGraph/Private/AnimGraphNode_SequenceEvaluator.cpp
--- a/Engine/Source/Editor/AnimGraph/Private/AnimGraphNode_SequenceEvaluator.cpp
+++ b/Engine/Source/Editor/AnimGraph/Private/AnimGraphNode_SequenceEvaluator.cpp
@@ -4,7 +4,7 @@
 #include "ToolMenus.h"
 
 #include "Kismet2/CompilerResultsLog.h"
-#include "GraphEditorActions.h"
+#include "AnimGraphCommands.h"
 #include "Animation/AnimComposite.h"
 #include "Animation/AnimSequence.h"
 #include "AnimGraphNode_SequenceEvaluator.h"
@@ -164,8 +164,8 @@ void UAnimGraphNode_SequenceEvaluator::GetNodeContextMenuActions(UToolMenu* Menu
 		// add an option to convert to a regular sequence player
 		{
 			FToolMenuSection& Section = Menu->AddSection("AnimGraphNodeSequenceEvaluator", NSLOCTEXT("A3Nodes", "SequenceEvaluatorHeading", "Sequence Evaluator"));
-			Section.AddMenuEntry(FGraphEditorCommands::Get().OpenRelatedAsset);
-			Section.AddMenuEntry(FGraphEditorCommands::Get().ConvertToSeqPlayer);
+			Section.AddMenuEntry(FAnimGraphCommands::Get().OpenRelatedAsset);
+			Section.AddMenuEntry(FAnimGraphCommands::Get().ConvertToSeqPlayer);
 		}
 	}
 }
diff --git a/Engine/Source/Editor/AnimGraph/Private/AnimGraphNode_SequencePlayer.cpp b/Engine/Source/Editor/AnimGraph/Private/AnimGraphNode_SequencePlayer.cpp
--- a/Engine/Source/Editor/AnimGraph/Private/AnimGraphNode_SequencePlayer.cpp
+++ b/Engine/Source/Editor/AnimGraph/Private/AnimGraphNode_SequencePlayer.cpp
@@ -6,7 +6,7 @@
 #include "ToolMenus.h"
 
 #include "Kismet2/CompilerResultsLog.h"
-#include "GraphEditorActions.h"
+#include "AnimGraphCommands.h"
 #include "ARFilter.h"
 #include "AssetRegistryModule.h"
 #include "BlueprintActionFilter.h"
@@ -402,8 +402,8 @@ void UAnimGraphNode_SequencePlayer::GetNodeContextMenuActions(UToolMenu* Menu, U
 		// add an option to convert to single frame
 		{
 			FToolMenuSection& Section = Menu->AddSection("AnimGraphNodeSequencePlayer", NSLOCTEXT("A3Nodes", "SequencePlayerHeading", "Sequence Player"));
-			Section.AddMenuEntry(FGraphEditorCommands::Get().OpenRelatedAsset);
-			Section.AddMenuEntry(FGraphEditorCommands::Get().ConvertToSeqEvaluator);
+			Section.AddMenuEntry(FAnimGraphCommands::Get().OpenRelatedAsset);
+			Section.AddMenuEntry(FAnimGraphCommands::Get().ConvertToSeqEvaluator);
 		}
 	}
 }
diff --git a/Engine/Source/Editor/AnimGraph/Private/AnimGraphNode_StateMachineBase.cpp b/Engine/Source/Editor/AnimGraph/Private/AnimGraphNode_StateMachineBase.cpp
--- a/Engine/Source/Editor/AnimGraph/Private/AnimGraphNode_StateMachineBase.cpp
+++ b/Engine/Source/Editor/AnimGraph/Private/AnimGraphNode_StateMachineBase.cpp
@@ -213,6 +213,11 @@ void UAnimGraphNode_StateMachineBase::OnRenameNode(const FString& NewName)
 	FBlueprintEditorUtils::RenameGraph(EditorStateMachineGraph, NewName);
 }
 
+TArray<UEdGraph*> UAnimGraphNode_StateMachineBase::GetSubGraphs() const
+{
+	return TArray<UEdGraph*>( { EditorStateMachineGraph } );
+}
+
 void UAnimGraphNode_StateMachineBase::OnProcessDuringCompilation(IAnimBlueprintCompilationContext& InCompilationContext, IAnimBlueprintGeneratedClassCompiledData& OutCompiledData)
 {
 	struct FMachineCreator
@@ -495,8 +500,12 @@ void UAnimGraphNode_StateMachineBase::OnProcessDuringCompilation(IAnimBlueprintC
 
 		// Add indices to all player and layer nodes
 		TArray<UEdGraph*> GraphsToCheck;
-		GraphsToCheck.Add(StateNode->GetBoundGraph());
-		StateNode->GetBoundGraph()->GetAllChildrenGraphs(GraphsToCheck);
+		TArray<UEdGraph*> SubGraphs = StateNode->GetSubGraphs();
+		GraphsToCheck.Append(SubGraphs);
+		for(UEdGraph* SubGraph : SubGraphs)
+		{
+			SubGraph->GetAllChildrenGraphs(GraphsToCheck);
+		}
 
 		TArray<UAnimGraphNode_LinkedAnimLayer*> LinkedAnimLayerNodes;
 		TArray<UAnimGraphNode_AssetPlayerBase*> AssetPlayerNodes;
diff --git a/Engine/Source/Editor/AnimGraph/Private/AnimStateNodeBase.cpp b/Engine/Source/Editor/AnimGraph/Private/AnimStateNodeBase.cpp
--- a/Engine/Source/Editor/AnimGraph/Private/AnimStateNodeBase.cpp
+++ b/Engine/Source/Editor/AnimGraph/Private/AnimStateNodeBase.cpp
@@ -48,22 +48,22 @@ void UAnimStateNodeBase::PostPasteNode()
 {
 	Super::PostPasteNode();
 
-	if (UEdGraph* BoundGraph = GetBoundGraph())
+	for(UEdGraph* SubGraph : GetSubGraphs())
 	{
 		// Add the new graph as a child of our parent graph
 		UEdGraph* ParentGraph = GetGraph();
 
-		if(ParentGraph->SubGraphs.Find(BoundGraph) == INDEX_NONE)
+		if(ParentGraph->SubGraphs.Find(SubGraph) == INDEX_NONE)
 		{
-			ParentGraph->SubGraphs.Add(BoundGraph);
+			ParentGraph->SubGraphs.Add(SubGraph);
 		}
 
 		//@TODO: CONDUIT: Merge conflict - May no longer be necessary due to other changes?
-//		FBlueprintEditorUtils::RenameGraphWithSuggestion(BoundGraph, NameValidator, GetDesiredNewNodeName());
+//		FBlueprintEditorUtils::RenameGraphWithSuggestion(SubGraph, NameValidator, GetDesiredNewNodeName());
 		//@ENDTODO
 
 		// Restore transactional flag that is lost during copy/paste process
-		BoundGraph->SetFlags(RF_Transactional);
+		SubGraph->SetFlags(RF_Transactional);
 
 		UBlueprint* Blueprint = FBlueprintEditorUtils::FindBlueprintForGraphChecked(ParentGraph);
 		FBlueprintEditorUtils::MarkBlueprintAsStructurallyModified(Blueprint);
@@ -72,7 +72,9 @@ void UAnimStateNodeBase::PostPasteNode()
 
 UObject* UAnimStateNodeBase::GetJumpTargetForDoubleClick() const
 {
-	return GetBoundGraph();
+	TArray<UEdGraph*> SubGraphs = GetSubGraphs();
+	check(SubGraphs.Num() > 0);
+	return SubGraphs[0];
 }
 
 bool UAnimStateNodeBase::CanJumpToDefinition() const
@@ -95,7 +97,9 @@ bool UAnimStateNodeBase::CanCreateUnderSpecifiedSchema(const UEdGraphSchema* Sch
 
 void UAnimStateNodeBase::OnRenameNode(const FString& NewName)
 {
-	FBlueprintEditorUtils::RenameGraph(GetBoundGraph(), NewName);
+	TArray<UEdGraph*> SubGraphs = GetSubGraphs();
+	check(SubGraphs.Num() > 0);
+	FBlueprintEditorUtils::RenameGraph(SubGraphs[0], NewName);
 }
 
 TSharedPtr<class INameValidatorInterface> UAnimStateNodeBase::MakeNameValidator() const
diff --git a/Engine/Source/Editor/AnimGraph/Private/AnimStateTransitionNode.cpp b/Engine/Source/Editor/AnimGraph/Private/AnimStateTransitionNode.cpp
--- a/Engine/Source/Editor/AnimGraph/Private/AnimStateTransitionNode.cpp
+++ b/Engine/Source/Editor/AnimGraph/Private/AnimStateTransitionNode.cpp
@@ -572,7 +572,7 @@ void UAnimStateTransitionNode::DestroyNode()
 }
 
 /** Returns true if this nodes BoundGraph is shared with another node in the parent graph */
-bool UAnimStateTransitionNode::IsBoundGraphShared()
+bool UAnimStateTransitionNode::IsBoundGraphShared() const
 {
 	if (BoundGraph)
 	{
@@ -641,6 +641,17 @@ void UAnimStateTransitionNode::ValidateNodeDuringCompilation(class FCompilerResu
 	}
 }
 
+TArray<UEdGraph*> UAnimStateTransitionNode::GetSubGraphs() const
+{ 
+	TArray<UEdGraph*> SubGraphs;
+	if(!IsBoundGraphShared())
+	{
+		SubGraphs.Add(BoundGraph);
+	}
+	SubGraphs.Add(CustomTransitionGraph);
+	return SubGraphs; 
+}
+
 //////////////////////////////////////////////////////////////////////////
 // IAnimStateTransitionNodeSharedDataHelper
 
diff --git a/Engine/Source/Editor/AnimGraph/Public/AnimGraphCommands.h b/Engine/Source/Editor/AnimGraph/Public/AnimGraphCommands.h
--- a/Engine/Source/Editor/AnimGraph/Public/AnimGraphCommands.h
+++ b/Engine/Source/Editor/AnimGraph/Public/AnimGraphCommands.h
@@ -22,4 +22,30 @@ public:
 	// Toggle pose watching for a given node
 	TSharedPtr<FUICommandInfo> TogglePoseWatch;
 
+	// SkeletalControl specific commands
+	TSharedPtr< FUICommandInfo > SelectBone;
+	// Blend list options
+	TSharedPtr< FUICommandInfo > AddBlendListPin;
+	TSharedPtr< FUICommandInfo > RemoveBlendListPin;
+
+	// options for sequence/evaluator converter
+	TSharedPtr< FUICommandInfo > ConvertToSeqEvaluator;
+	TSharedPtr< FUICommandInfo > ConvertToSeqPlayer;
+
+	// options for blendspace converter
+	TSharedPtr< FUICommandInfo > ConvertToBSEvaluator;
+	TSharedPtr< FUICommandInfo > ConvertToBSPlayer;
+	TSharedPtr< FUICommandInfo > ConvertToBSGraph;
+
+	// options for aimoffset converter
+	TSharedPtr< FUICommandInfo > ConvertToAimOffsetLookAt;
+	TSharedPtr< FUICommandInfo > ConvertToAimOffsetSimple;
+	TSharedPtr< FUICommandInfo > ConvertToAimOffsetGraph;
+
+	// options for sequence/evaluator converter
+	TSharedPtr< FUICommandInfo > ConvertToPoseBlender;
+	TSharedPtr< FUICommandInfo > ConvertToPoseByName;
+
+	// option for opening the asset related to the graph node
+	TSharedPtr< FUICommandInfo > OpenRelatedAsset;
 };
diff --git a/Engine/Source/Editor/AnimGraph/Public/AnimGraphNode_Base.h b/Engine/Source/Editor/AnimGraph/Public/AnimGraphNode_Base.h
--- a/Engine/Source/Editor/AnimGraph/Public/AnimGraphNode_Base.h
+++ b/Engine/Source/Editor/AnimGraph/Public/AnimGraphNode_Base.h
@@ -25,6 +25,8 @@ class USkeletalMeshComponent;
 class IAnimBlueprintCompilerHandlerCollection;
 class IAnimBlueprintGeneratedClassCompiledData;
 class IAnimBlueprintCompilationContext;
+class IAnimBlueprintCopyTermDefaultsContext;
+class IAnimBlueprintNodeCopyTermDefaultsContext;
 
 struct FPoseLinkMappingRecord
 {
@@ -365,6 +367,11 @@ class ANIMGRAPH_API UAnimGraphNode_Base : public UK2Node
 	// @return wether to show graph attribute icons on pins for this node.
 	virtual bool ShouldShowAttributesOnPins() const { return true; }
 
+	// Some anim graph nodes can act as proxies to outer nodes (e.g. sink nodes in nested graphs).
+	// Returning that outer node here allows attribute visualization to be forwarded to the inner node from the outer.
+	// @return the proxy for attribute display
+	virtual const UAnimGraphNode_Base* GetProxyNodeForAttributes() const { return this; }
+
 protected:
 	friend class FAnimBlueprintCompilerContext;
 	friend class FAnimGraphNodeDetails;
@@ -393,6 +400,13 @@ protected:
 	// Process this node's data during compilation (override point)
 	virtual void OnProcessDuringCompilation(IAnimBlueprintCompilationContext& InCompilationContext, IAnimBlueprintGeneratedClassCompiledData& OutCompiledData) {}
 
+	// Copy this node's data during the last phase of compilation where term defaults are copied to the new CDO
+	void CopyTermDefaultsToDefaultObject(IAnimBlueprintCopyTermDefaultsContext& InCompilationContext, IAnimBlueprintNodeCopyTermDefaultsContext& InPerNodeContext, IAnimBlueprintGeneratedClassCompiledData& OutCompiledData);
+
+	// Override point for CopyTermDefaultsToDefaultObject
+	// Copy this node's data during the last phase of compilation where term defaults are copied to the new CDO
+	virtual void OnCopyTermDefaultsToDefaultObject(IAnimBlueprintCopyTermDefaultsContext& InCompilationContext, IAnimBlueprintNodeCopyTermDefaultsContext& InPerNodeContext, IAnimBlueprintGeneratedClassCompiledData& OutCompiledData) {}
+
 	// Allocates or reallocates pins
 	void InternalPinCreation(TArray<UEdGraphPin*>* OldPins);
 
diff --git a/Engine/Source/Editor/AnimGraph/Public/AnimGraphNode_BlendSpaceGraph.h b/Engine/Source/Editor/AnimGraph/Public/AnimGraphNode_BlendSpaceGraph.h
new file mode 100644
--- /dev/null
+++ b/Engine/Source/Editor/AnimGraph/Public/AnimGraphNode_BlendSpaceGraph.h
@@ -0,0 +1,30 @@
+// Copyright Epic Games, Inc. All Rights Reserved.
+
+#pragma once
+
+#include "CoreMinimal.h"
+#include "AnimGraphNode_BlendSpaceGraphBase.h"
+#include "AnimNodes/AnimNode_BlendSpaceGraph.h"
+#include "AnimGraphNode_BlendSpaceGraph.generated.h"
+
+class FBlueprintActionDatabaseRegistrar;
+class FCompilerResultsLog;
+
+UCLASS()
+class ANIMGRAPH_API UAnimGraphNode_BlendSpaceGraph : public UAnimGraphNode_BlendSpaceGraphBase
+{
+	GENERATED_BODY()
+
+private:
+	UPROPERTY(EditAnywhere, Category=Settings)
+	FAnimNode_BlendSpaceGraph Node;
+
+	// UEdGraphNode interface
+	virtual FText GetNodeTitle(ENodeTitleType::Type TitleType) const override;
+	
+	// UK2Node interface
+	virtual void GetMenuActions(FBlueprintActionDatabaseRegistrar& ActionRegistrar) const override;
+
+	// UAnimGraphNode_Base interface
+	virtual void BakeDataDuringCompilation(FCompilerResultsLog& MessageLog) override;
+};
diff --git a/Engine/Source/Editor/AnimGraph/Public/AnimGraphNode_BlendSpaceGraphBase.h b/Engine/Source/Editor/AnimGraph/Public/AnimGraphNode_BlendSpaceGraphBase.h
new file mode 100644
--- /dev/null
+++ b/Engine/Source/Editor/AnimGraph/Public/AnimGraphNode_BlendSpaceGraphBase.h
@@ -0,0 +1,99 @@
+// Copyright Epic Games, Inc. All Rights Reserved.
+
+#pragma once
+
+#include "CoreMinimal.h"
+#include "UObject/ObjectMacros.h"
+#include "AnimGraphNode_Base.h"
+#include "Containers/ArrayView.h"
+#include "AnimGraphNode_BlendSpaceGraphBase.generated.h"
+
+class FBlueprintActionDatabaseRegistrar;
+class IAnimBlueprintCopyTermDefaultsContext;
+class IAnimBlueprintNodeCopyTermDefaultsContext;
+class IAnimBlueprintGeneratedClassCompiledData;
+class UBlendSpaceGraph;
+class UAnimationBlendSpaceSampleGraph;
+
+UCLASS()
+class ANIMGRAPH_API UAnimGraphNode_BlendSpaceGraphBase : public UAnimGraphNode_Base
+{
+	GENERATED_BODY()
+
+public:
+	UAnimGraphNode_BlendSpaceGraphBase();
+
+	// Access the graphs for each sample
+	TArrayView<UEdGraph* const> GetGraphs() const { return Graphs; }
+
+	// Access the 'dummy' blendspace graph
+	UBlendSpaceGraph* GetBlendSpaceGraph() const { return BlendSpaceGraph; }
+
+	// Adds a new graph to the internal array
+	UAnimationBlendSpaceSampleGraph* AddGraph(FName InSampleName);
+
+	// Removes the graph at the specified index
+	void RemoveGraph(int32 InSampleIndex);
+
+	// Setup this node from the specified asset
+	void SetupFromAsset(UBlendSpaceBase* InBlendSpace, bool bInIsTemplateNode);
+
+	// UEdGraphNode interface
+	virtual void PostPlacedNewNode() override;
+
+	// @return the sync group name assigned to this node
+	FName GetSyncGroupName() const;
+
+protected:
+	// Get the name of the blendspace graph
+	FString GetBlendSpaceGraphName() const;
+
+	// Get the name of the blendspace
+	FString GetBlendSpaceName() const;
+
+	// Setup this node from the specified class
+	void SetupFromClass(TSubclassOf<UBlendSpaceBase> InBlendSpaceClass, bool bInIsTemplateNode);
+
+	// Internal blendspace
+	UPROPERTY()
+	UBlendSpaceBase* BlendSpace;
+
+	// Blendspace class, for template nodes
+	UPROPERTY()
+	TSubclassOf<UBlendSpaceBase> BlendSpaceClass;
+
+	// Dummy blendspace graph (used for navigation only)
+	UPROPERTY()
+	UBlendSpaceGraph* BlendSpaceGraph;
+
+	// Linked animation graphs for sample points
+	UPROPERTY()
+	TArray<UEdGraph*> Graphs;
+
+protected:
+	// UEdGraphNode interface
+	virtual FText GetMenuCategory() const override;
+	virtual FLinearColor GetNodeTitleColor() const override;
+	virtual FText GetTooltipText() const override;
+	virtual UObject* GetJumpTargetForDoubleClick() const override;
+	virtual void JumpToDefinition() const override;
+	virtual TArray<UEdGraph*> GetSubGraphs() const override;
+	virtual void DestroyNode() override;
+	virtual void OnRenameNode(const FString& NewName) override;
+	virtual TSharedPtr<INameValidatorInterface> MakeNameValidator() const override;
+	virtual void PostPasteNode() override;
+	virtual void CustomizePinData(UEdGraphPin* Pin, FName SourcePropertyName, int32 ArrayIndex) const override;
+	virtual void PostProcessPinName(const UEdGraphPin* Pin, FString& DisplayName) const override;
+
+	// UAnimGraphNode_Base interface
+	virtual void OnProcessDuringCompilation(IAnimBlueprintCompilationContext& InCompilationContext, IAnimBlueprintGeneratedClassCompiledData& OutCompiledData) override;
+	virtual void OnCopyTermDefaultsToDefaultObject(IAnimBlueprintCopyTermDefaultsContext& InCompilationContext, IAnimBlueprintNodeCopyTermDefaultsContext& InPerNodeContext, IAnimBlueprintGeneratedClassCompiledData& OutCompiledData) override;
+	virtual void CustomizeDetails(IDetailLayoutBuilder& InDetailBuilder) override;
+	virtual void GetInputLinkAttributes(FNodeAttributeArray& OutAttributes) const override;
+
+	// UK2Node interface
+	virtual void PreloadRequiredAssets() override;
+
+	// Helper function for compilation
+	UAnimGraphNode_Base* ExpandGraphAndProcessNodes(UEdGraph* SourceGraph, UAnimGraphNode_Base* SourceRootNode, IAnimBlueprintCompilationContext& InCompilationContext, IAnimBlueprintGeneratedClassCompiledData& OutCompiledData);
+};
diff --git a/Engine/Source/Editor/AnimGraph/Public/AnimGraphNode_LinkedAnimGraphBase.h b/Engine/Source/Editor/AnimGraph/Public/AnimGraphNode_LinkedAnimGraphBase.h
--- a/Engine/Source/Editor/AnimGraph/Public/AnimGraphNode_LinkedAnimGraphBase.h
+++ b/Engine/Source/Editor/AnimGraph/Public/AnimGraphNode_LinkedAnimGraphBase.h
@@ -36,6 +36,7 @@ public:
 	virtual FPoseLinkMappingRecord GetLinkIDLocation(const UScriptStruct* NodeType, UEdGraphPin* SourcePin) override;
 	virtual void GetOutputLinkAttributes(FNodeAttributeArray& OutAttributes) const override;
 	virtual bool ShouldShowAttributesOnPins() const override { return false; }
+	virtual void OnCopyTermDefaultsToDefaultObject(IAnimBlueprintCopyTermDefaultsContext& InCompilationContext, IAnimBlueprintNodeCopyTermDefaultsContext& InPerNodeContext, IAnimBlueprintGeneratedClassCompiledData& OutCompiledData) override;
 
 	// UAnimGraphNode_CustomProperty interface
 	virtual bool IsStructuralProperty(FProperty* InProperty) const override;
diff --git a/Engine/Source/Editor/AnimGraph/Public/AnimGraphNode_LinkedInputPose.h b/Engine/Source/Editor/AnimGraph/Public/AnimGraphNode_LinkedInputPose.h
--- a/Engine/Source/Editor/AnimGraph/Public/AnimGraphNode_LinkedInputPose.h
+++ b/Engine/Source/Editor/AnimGraph/Public/AnimGraphNode_LinkedInputPose.h
@@ -87,6 +87,7 @@ public:
 	virtual void CustomizeDetails(IDetailLayoutBuilder& DetailBuilder) override;
 	virtual void GetOutputLinkAttributes(FNodeAttributeArray& OutAttributes) const override;
 	virtual bool ShouldShowAttributesOnPins() const override { return false; }
+	virtual void OnCopyTermDefaultsToDefaultObject(IAnimBlueprintCopyTermDefaultsContext& InCompilationContext, IAnimBlueprintNodeCopyTermDefaultsContext& InPerNodeContext, IAnimBlueprintGeneratedClassCompiledData& OutCompiledData) override;
 
 	/** Make a name widget for this linked input pose node */
 	TSharedRef<SWidget> MakeNameWidget(IDetailLayoutBuilder& DetailBuilder);
diff --git a/Engine/Source/Editor/AnimGraph/Public/AnimGraphNode_RotationOffsetBlendSpaceGraph.h b/Engine/Source/Editor/AnimGraph/Public/AnimGraphNode_RotationOffsetBlendSpaceGraph.h
new file mode 100644
--- /dev/null
+++ b/Engine/Source/Editor/AnimGraph/Public/AnimGraphNode_RotationOffsetBlendSpaceGraph.h
@@ -0,0 +1,35 @@
+// Copyright Epic Games, Inc. All Rights Reserved.
+
+#pragma once
+
+#include "CoreMinimal.h"
+#include "AnimGraphNode_BlendSpaceGraphBase.h"
+#include "AnimNodes/AnimNode_RotationOffsetBlendSpaceGraph.h"
+#include "AnimGraphNode_RotationOffsetBlendSpaceGraph.generated.h"
+
+class FBlueprintActionDatabaseRegistrar;
+class FCompilerResultsLog;
+
+UCLASS()
+class ANIMGRAPH_API UAnimGraphNode_RotationOffsetBlendSpaceGraph : public UAnimGraphNode_BlendSpaceGraphBase
+{
+	GENERATED_BODY()
+
+private:
+	UPROPERTY(EditAnywhere, Category=Settings)
+	FAnimNode_RotationOffsetBlendSpaceGraph Node;
+
+	// UObject interface
+	virtual void PostEditChangeProperty(FPropertyChangedEvent& PropertyChangedEvent) override;
+
+	// UEdGraphNode interface
+	virtual FText GetNodeTitle(ENodeTitleType::Type TitleType) const override;
+	
+	// UK2Node interface
+	virtual void GetMenuActions(FBlueprintActionDatabaseRegistrar& ActionRegistrar) const override;
+	virtual void CustomizePinData(UEdGraphPin* Pin, FName SourcePropertyName, int32 ArrayIndex) const override;
+
+	// UAnimGraphNode_Base interface
+	virtual void BakeDataDuringCompilation(FCompilerResultsLog& MessageLog) override;
+	virtual void CustomizeDetails(IDetailLayoutBuilder& DetailBuilder) override;
+};
diff --git a/Engine/Source/Editor/AnimGraph/Public/AnimGraphNode_StateMachineBase.h b/Engine/Source/Editor/AnimGraph/Public/AnimGraphNode_StateMachineBase.h
--- a/Engine/Source/Editor/AnimGraph/Public/AnimGraphNode_StateMachineBase.h
+++ b/Engine/Source/Editor/AnimGraph/Public/AnimGraphNode_StateMachineBase.h
@@ -32,6 +32,7 @@ class ANIMGRAPH_API UAnimGraphNode_StateMachineBase : public UAnimGraphNode_Base
 	virtual TSharedPtr<class INameValidatorInterface> MakeNameValidator() const override;
 	virtual FString GetDocumentationLink() const override;
 	virtual void OnRenameNode(const FString& NewName) override;
+	virtual TArray<UEdGraph*> GetSubGraphs() const override;
 	// End of UEdGraphNode interface
 
 	// UAnimGraphNode_Base interface
diff --git a/Engine/Source/Editor/AnimGraph/Public/AnimStateConduitNode.h b/Engine/Source/Editor/AnimGraph/Public/AnimStateConduitNode.h
--- a/Engine/Source/Editor/AnimGraph/Public/AnimStateConduitNode.h
+++ b/Engine/Source/Editor/AnimGraph/Public/AnimStateConduitNode.h
@@ -30,6 +30,7 @@ public:
 	virtual void PostPlacedNewNode() override;
 	virtual void DestroyNode() override;
 	virtual void ValidateNodeDuringCompilation(class FCompilerResultsLog& MessageLog) const override;
+	virtual TArray<UEdGraph*> GetSubGraphs() const override { return TArray<UEdGraph*>( { BoundGraph } ); }
 	//~ End UEdGraphNode Interface
 
 	//~ Begin UAnimStateNodeBase Interface
diff --git a/Engine/Source/Editor/AnimGraph/Public/AnimStateNode.h b/Engine/Source/Editor/AnimGraph/Public/AnimStateNode.h
--- a/Engine/Source/Editor/AnimGraph/Public/AnimStateNode.h
+++ b/Engine/Source/Editor/AnimGraph/Public/AnimStateNode.h
@@ -59,6 +59,7 @@ public:
 	virtual void PostPasteNode() override;
 	virtual void PostPlacedNewNode() override;
 	virtual void DestroyNode() override;
+	virtual TArray<UEdGraph*> GetSubGraphs() const override { return TArray<UEdGraph*>( { BoundGraph } ); }
 	//~ End UEdGraphNode Interface
 	
 	//~ Begin UAnimStateNodeBase Interface
diff --git a/Engine/Source/Editor/AnimGraph/Public/AnimStateTransitionNode.h b/Engine/Source/Editor/AnimGraph/Public/AnimStateTransitionNode.h
--- a/Engine/Source/Editor/AnimGraph/Public/AnimStateTransitionNode.h
+++ b/Engine/Source/Editor/AnimGraph/Public/AnimStateTransitionNode.h
@@ -124,6 +124,7 @@ class UAnimStateTransitionNode : public UAnimStateNodeBase
 	virtual void PostPlacedNewNode() override;
 	virtual void DestroyNode() override;
 	virtual void ValidateNodeDuringCompilation(class FCompilerResultsLog& MessageLog) const override;
+	virtual TArray<UEdGraph*> GetSubGraphs() const override;
 	//~ End UEdGraphNode Interface
 
 	//~ Begin UAnimStateNodeBase Interface
@@ -140,7 +141,7 @@ class UAnimStateTransitionNode : public UAnimStateNodeBase
 	ANIMGRAPH_API UAnimStateNodeBase* GetNextState() const;
 	ANIMGRAPH_API void CreateConnections(UAnimStateNodeBase* PreviousState, UAnimStateNodeBase* NextState);
 	
-	ANIMGRAPH_API bool IsBoundGraphShared();
+	ANIMGRAPH_API bool IsBoundGraphShared() const;
 
 	ANIMGRAPH_API void MakeRulesShareable(FString ShareName);
 	ANIMGRAPH_API void MakeCrossfadeShareable(FString ShareName);
diff --git a/Engine/Source/Editor/AnimGraph/Public/IAnimBlueprintCopyTermDefaultsContext.h b/Engine/Source/Editor/AnimGraph/Public/IAnimBlueprintCopyTermDefaultsContext.h
--- a/Engine/Source/Editor/AnimGraph/Public/IAnimBlueprintCopyTermDefaultsContext.h
+++ b/Engine/Source/Editor/AnimGraph/Public/IAnimBlueprintCopyTermDefaultsContext.h
@@ -26,3 +26,36 @@ protected:
 	// Get the currently-compiled anim blueprint
 	virtual const UAnimBlueprint* GetAnimBlueprintImpl() const = 0;
 };
+
+
+// Interface passed to per-node CopyTermDefaults override point
+class IAnimBlueprintNodeCopyTermDefaultsContext
+{
+public:	
+	virtual ~IAnimBlueprintNodeCopyTermDefaultsContext() {}
+
+	// Get the property that we are writing to
+	const FProperty* GetTargetProperty() const { return GetTargetPropertyImpl(); }
+
+	// Get the destination ptr (the node) that we are writing to
+	uint8* GetDestinationPtr() const { return GetDestinationPtrImpl(); }
+
+	// Get the source ptr (the node in the anim graph node) that we are reading from
+	const uint8* GetSourcePtr() const { return GetSourcePtrImpl(); }
+
+	// Get the property index for this node
+	int32 GetNodePropertyIndex() const { return GetNodePropertyIndexImpl(); }
+
+protected:
+	// Get the property that we are writing to
+	virtual const FProperty* GetTargetPropertyImpl() const = 0;
+
+	// Get the destination ptr (the node) that we are writing to
+	virtual uint8* GetDestinationPtrImpl() const = 0;
+
+	// Get the source ptr (the node in the anim graph node) that we are reading from
+	virtual const uint8* GetSourcePtrImpl() const = 0;
+
+	// Get the property index for this node
+	virtual int32 GetNodePropertyIndexImpl() const  = 0;
+};
\ No newline at end of file
diff --git a/Engine/Source/Editor/AnimGraph/Public/IAnimBlueprintGeneratedClassCompiledData.h b/Engine/Source/Editor/AnimGraph/Public/IAnimBlueprintGeneratedClassCompiledData.h
--- a/Engine/Source/Editor/AnimGraph/Public/IAnimBlueprintGeneratedClassCompiledData.h
+++ b/Engine/Source/Editor/AnimGraph/Public/IAnimBlueprintGeneratedClassCompiledData.h
@@ -42,4 +42,9 @@ public:
 
 	// Get the currently-compiled classes graph asset player information
 	virtual TMap<FName, FGraphAssetPlayerInformation>& GetGraphAssetPlayerInformation() const = 0;
+
+	// Add a class-internal blendspace
+	// @param	InSourceBlendSpace	The blendspace to duplicate
+	// @return a duplicate of the blendspace outered to the class
+	virtual UBlendSpaceBase* AddBlendSpace(UBlendSpaceBase* InSourceBlendSpace) = 0;
 };
\ No newline at end of file
diff --git a/Engine/Source/Editor/AnimationBlueprintEditor/AnimationBlueprintEditor.Build.cs b/Engine/Source/Editor/AnimationBlueprintEditor/AnimationBlueprintEditor.Build.cs
--- a/Engine/Source/Editor/AnimationBlueprintEditor/AnimationBlueprintEditor.Build.cs
+++ b/Engine/Source/Editor/AnimationBlueprintEditor/AnimationBlueprintEditor.Build.cs
@@ -35,7 +35,8 @@ public class AnimationBlueprintEditor : ModuleRules
                 "InputCore",
 				"KismetWidgets",
 				"AnimGraph",
-                "PropertyEditor",
+				"AnimGraphRuntime",
+				"PropertyEditor",
 				"EditorWidgets",
                 "BlueprintGraph",
                 "RHI",
diff --git a/Engine/Source/Editor/AnimationBlueprintEditor/Private/AnimationBlueprintEditor.cpp b/Engine/Source/Editor/AnimationBlueprintEditor/Private/AnimationBlueprintEditor.cpp
--- a/Engine/Source/Editor/AnimationBlueprintEditor/Private/AnimationBlueprintEditor.cpp
+++ b/Engine/Source/Editor/AnimationBlueprintEditor/Private/AnimationBlueprintEditor.cpp
@@ -80,6 +80,11 @@
 #include "AnimStateEntryNode.h"
 #include "PersonaUtils.h"
 #include "Subsystems/AssetEditorSubsystem.h"
+#include "BlendSpaceDocumentTabFactory.h"
+#include "BlendSpaceGraph.h"
+#include "TabPayload_BlendSpaceGraph.h"
+#include "AnimGraphNode_RotationOffsetBlendSpaceGraph.h"
+#include "AnimGraphNode_BlendSpaceGraph.h"
 
 #define LOCTEXT_NAMESPACE "AnimationBlueprintEditor"
 
@@ -290,6 +295,9 @@ void FAnimationBlueprintEditor::InitAnimationBlueprintEditor(const EToolkitMode:
 
 	CommonInitialization(AnimBlueprints, /*bShouldOpenInDefaultsMode=*/ false);
 
+	// Register document editor for blendspaces
+	DocumentManager->RegisterDocumentFactory(MakeShared<FBlendSpaceDocumentTabFactory>(SharedThis(this)));
+
 	if(InAnimBlueprint->BlueprintType == BPTYPE_Interface)
 	{
 		AddApplicationMode(
@@ -481,6 +489,60 @@ void FAnimationBlueprintEditor::OnCreateGraphEditorCommands(TSharedPtr<FUIComman
 {
 	GraphEditorCommandsList->MapAction(FAnimGraphCommands::Get().TogglePoseWatch,
 		FExecuteAction::CreateSP(this, &FAnimationBlueprintEditor::OnTogglePoseWatch));
+
+	GraphEditorCommandsList->MapAction( FAnimGraphCommands::Get().AddBlendListPin,
+		FExecuteAction::CreateSP( this, &FAnimationBlueprintEditor::OnAddPosePin ),
+		FCanExecuteAction::CreateSP( this, &FAnimationBlueprintEditor::CanAddPosePin )
+		);
+
+	GraphEditorCommandsList->MapAction( FAnimGraphCommands::Get().RemoveBlendListPin,
+		FExecuteAction::CreateSP( this, &FAnimationBlueprintEditor::OnRemovePosePin ),
+		FCanExecuteAction::CreateSP( this, &FAnimationBlueprintEditor::CanRemovePosePin )
+		);
+
+	GraphEditorCommandsList->MapAction( FAnimGraphCommands::Get().ConvertToSeqEvaluator,
+		FExecuteAction::CreateSP( this, &FAnimationBlueprintEditor::OnConvertToSequenceEvaluator )
+		);
+
+	GraphEditorCommandsList->MapAction( FAnimGraphCommands::Get().ConvertToSeqPlayer,
+		FExecuteAction::CreateSP( this, &FAnimationBlueprintEditor::OnConvertToSequencePlayer )
+		);
+
+	GraphEditorCommandsList->MapAction( FAnimGraphCommands::Get().ConvertToBSEvaluator,
+		FExecuteAction::CreateSP( this, &FAnimationBlueprintEditor::OnConvertToBlendSpaceEvaluator )
+		);
+
+	GraphEditorCommandsList->MapAction( FAnimGraphCommands::Get().ConvertToBSPlayer,
+		FExecuteAction::CreateSP( this, &FAnimationBlueprintEditor::OnConvertToBlendSpacePlayer )
+		);
+
+	GraphEditorCommandsList->MapAction( FAnimGraphCommands::Get().ConvertToBSGraph,
+		FExecuteAction::CreateSP( this, &FAnimationBlueprintEditor::OnConvertToBlendSpaceGraph )
+		);
+
+	GraphEditorCommandsList->MapAction(FAnimGraphCommands::Get().ConvertToAimOffsetLookAt,
+		FExecuteAction::CreateSP(this, &FAnimationBlueprintEditor::OnConvertToAimOffsetLookAt)
+	);
+
+	GraphEditorCommandsList->MapAction(FAnimGraphCommands::Get().ConvertToAimOffsetSimple,
+		FExecuteAction::CreateSP(this, &FAnimationBlueprintEditor::OnConvertToAimOffsetSimple)
+	);
+
+	GraphEditorCommandsList->MapAction(FAnimGraphCommands::Get().ConvertToAimOffsetGraph,
+		FExecuteAction::CreateSP(this, &FAnimationBlueprintEditor::OnConvertToAimOffsetGraph)
+	);	
+
+	GraphEditorCommandsList->MapAction(FAnimGraphCommands::Get().ConvertToPoseBlender,
+		FExecuteAction::CreateSP(this, &FAnimationBlueprintEditor::OnConvertToPoseBlender)
+		);
+
+	GraphEditorCommandsList->MapAction(FAnimGraphCommands::Get().ConvertToPoseByName,
+		FExecuteAction::CreateSP(this, &FAnimationBlueprintEditor::OnConvertToPoseByName)
+		);
+
+	GraphEditorCommandsList->MapAction( FAnimGraphCommands::Get().OpenRelatedAsset,
+		FExecuteAction::CreateSP( this, &FAnimationBlueprintEditor::OnOpenRelatedAsset )
+		);
 }
 
 
@@ -668,12 +730,27 @@ void FAnimationBlueprintEditor::OnTogglePoseWatch()
 	}
 }
 
+// Helper function for node conversions
+static void CopyPinData(UEdGraphNode* InOldNode, UEdGraphNode* InNewNode, const TCHAR* InPinName)
+{
+	UEdGraphPin* OldPin = InOldNode->FindPin(InPinName);
+	UEdGraphPin* NewPin = InNewNode->FindPin(InPinName);
+
+	if (ensure(OldPin && NewPin))
+	{
+		NewPin->MovePersistentDataFromOldPin(*OldPin);
+	}
+};
+
 void FAnimationBlueprintEditor::OnConvertToSequenceEvaluator()
 {
 	FGraphPanelSelectionSet SelectedNodes = GetSelectedNodes();
 
 	if (SelectedNodes.Num() > 0)
 	{
+		// convert to sequence evaluator
+		const FScopedTransaction Transaction( LOCTEXT("ConvertToSequenceEvaluator", "Convert to Single Frame Animation") );
+
 		for (auto NodeIter = SelectedNodes.CreateIterator(); NodeIter; ++NodeIter)
 		{
 			UAnimGraphNode_SequencePlayer* OldNode = Cast<UAnimGraphNode_SequencePlayer>(*NodeIter);
@@ -681,10 +758,10 @@ void FAnimationBlueprintEditor::OnConvertToSequenceEvaluator()
 			// see if sequence player
 			if ( OldNode && OldNode->Node.Sequence )
 			{
-				//const FScopedTransaction Transaction( LOCTEXT("ConvertToSequenceEvaluator", "Convert to Single Frame Animation") );
-
-				// convert to sequence evaluator
 				UEdGraph* TargetGraph = OldNode->GetGraph();
+				TargetGraph->Modify();
+				OldNode->Modify();
+
 				// create new evaluator
 				FGraphNodeCreator<UAnimGraphNode_SequenceEvaluator> NodeCreator(*TargetGraph);
 				UAnimGraphNode_SequenceEvaluator* NewNode = NodeCreator.CreateNode();
@@ -694,19 +771,11 @@ void FAnimationBlueprintEditor::OnConvertToSequenceEvaluator()
 				// get default data from old node to new node
 				FEdGraphUtilities::CopyCommonState(OldNode, NewNode);
 
-				UEdGraphPin* OldPosePin = OldNode->FindPin(TEXT("Pose"));
-				UEdGraphPin* NewPosePin = NewNode->FindPin(TEXT("Pose"));
-
-				if (ensure(OldPosePin && NewPosePin))
-				{
-					NewPosePin->MovePersistentDataFromOldPin(*OldPosePin);
-				}
+				CopyPinData(OldNode, NewNode, TEXT("Pose"));
 
 				// remove from selection and from graph
 				NodeIter.RemoveCurrent();
 				TargetGraph->RemoveNode(OldNode);
-
-				NewNode->Modify();
 			}
 		}
 
@@ -729,6 +798,8 @@ void FAnimationBlueprintEditor::OnConvertToSequencePlayer()
 	FGraphPanelSelectionSet SelectedNodes = GetSelectedNodes();
 	if (SelectedNodes.Num() > 0)
 	{
+		const FScopedTransaction Transaction( LOCTEXT("ConvertToSequencePlayer", "Convert to Sequence Player") );
+
 		for (auto NodeIter = SelectedNodes.CreateIterator(); NodeIter; ++NodeIter)
 		{
 			UAnimGraphNode_SequenceEvaluator* OldNode = Cast<UAnimGraphNode_SequenceEvaluator>(*NodeIter);
@@ -736,9 +807,11 @@ void FAnimationBlueprintEditor::OnConvertToSequencePlayer()
 			// see if sequence player
 			if ( OldNode && OldNode->Node.Sequence )
 			{
-				//const FScopedTransaction Transaction( LOCTEXT("ConvertToSequenceEvaluator", "Convert to Single Frame Animation") );
 				// convert to sequence player
 				UEdGraph* TargetGraph = OldNode->GetGraph();
+				TargetGraph->Modify();
+				OldNode->Modify();
+
 				// create new player
 				FGraphNodeCreator<UAnimGraphNode_SequencePlayer> NodeCreator(*TargetGraph);
 				UAnimGraphNode_SequencePlayer* NewNode = NodeCreator.CreateNode();
@@ -748,19 +821,11 @@ void FAnimationBlueprintEditor::OnConvertToSequencePlayer()
 				// get default data from old node to new node
 				FEdGraphUtilities::CopyCommonState(OldNode, NewNode);
 
-				UEdGraphPin* OldPosePin = OldNode->FindPin(TEXT("Pose"));
-				UEdGraphPin* NewPosePin = NewNode->FindPin(TEXT("Pose"));
-
-				if (ensure(OldPosePin && NewPosePin))
-				{
-					NewPosePin->MovePersistentDataFromOldPin(*OldPosePin);
-				}
+				CopyPinData(OldNode, NewNode, TEXT("Pose"));
 
 				// remove from selection and from graph
 				NodeIter.RemoveCurrent();
 				TargetGraph->RemoveNode(OldNode);
-
-				NewNode->Modify();
 			}
 		}
 
@@ -784,6 +849,8 @@ void FAnimationBlueprintEditor::OnConvertToBlendSpaceEvaluator()
 
 	if (SelectedNodes.Num() > 0)
 	{
+		const FScopedTransaction Transaction( LOCTEXT("ConvertToBlendSpaceEvaluator", "Convert to Single Frame Blend Space") );
+
 		for (auto NodeIter = SelectedNodes.CreateIterator(); NodeIter; ++NodeIter)
 		{
 			UAnimGraphNode_BlendSpacePlayer* OldNode = Cast<UAnimGraphNode_BlendSpacePlayer>(*NodeIter);
@@ -791,10 +858,11 @@ void FAnimationBlueprintEditor::OnConvertToBlendSpaceEvaluator()
 			// see if sequence player
 			if ( OldNode && OldNode->Node.BlendSpace )
 			{
-				//const FScopedTransaction Transaction( LOCTEXT("ConvertToSequenceEvaluator", "Convert to Single Frame Animation") );
-
 				// convert to sequence evaluator
 				UEdGraph* TargetGraph = OldNode->GetGraph();
+				TargetGraph->Modify();
+				OldNode->Modify();
+
 				// create new evaluator
 				FGraphNodeCreator<UAnimGraphNode_BlendSpaceEvaluator> NodeCreator(*TargetGraph);
 				UAnimGraphNode_BlendSpaceEvaluator* NewNode = NodeCreator.CreateNode();
@@ -804,36 +872,13 @@ void FAnimationBlueprintEditor::OnConvertToBlendSpaceEvaluator()
 				// get default data from old node to new node
 				FEdGraphUtilities::CopyCommonState(OldNode, NewNode);
 
-				UEdGraphPin* OldPosePin = OldNode->FindPin(TEXT("X"));
-				UEdGraphPin* NewPosePin = NewNode->FindPin(TEXT("X"));
-
-				if (ensure(OldPosePin && NewPosePin))
-				{
-					NewPosePin->MovePersistentDataFromOldPin(*OldPosePin);
-				}
-
-				OldPosePin = OldNode->FindPin(TEXT("Y"));
-				NewPosePin = NewNode->FindPin(TEXT("Y"));
-
-				if (ensure(OldPosePin && NewPosePin))
-				{
-					NewPosePin->MovePersistentDataFromOldPin(*OldPosePin);
-				}
-
-
-				OldPosePin = OldNode->FindPin(TEXT("Pose"));
-				NewPosePin = NewNode->FindPin(TEXT("Pose"));
-
-				if (ensure(OldPosePin && NewPosePin))
-				{
-					NewPosePin->MovePersistentDataFromOldPin(*OldPosePin);
-				}
+				CopyPinData(OldNode, NewNode, TEXT("X"));
+				CopyPinData(OldNode, NewNode, TEXT("Y"));
+				CopyPinData(OldNode, NewNode, TEXT("Pose"));
 
 				// remove from selection and from graph
 				NodeIter.RemoveCurrent();
 				TargetGraph->RemoveNode(OldNode);
-
-				NewNode->Modify();
 			}
 		}
 
@@ -855,6 +900,8 @@ void FAnimationBlueprintEditor::OnConvertToBlendSpacePlayer()
 	FGraphPanelSelectionSet SelectedNodes = GetSelectedNodes();
 	if (SelectedNodes.Num() > 0)
 	{
+		const FScopedTransaction Transaction( LOCTEXT("ConvertToBlendSpacePlayer", "Convert to Blend Space Player") );
+
 		for (auto NodeIter = SelectedNodes.CreateIterator(); NodeIter; ++NodeIter)
 		{
 			UAnimGraphNode_BlendSpaceEvaluator* OldNode = Cast<UAnimGraphNode_BlendSpaceEvaluator>(*NodeIter);
@@ -862,9 +909,11 @@ void FAnimationBlueprintEditor::OnConvertToBlendSpacePlayer()
 			// see if sequence player
 			if ( OldNode && OldNode->Node.BlendSpace )
 			{
-				//const FScopedTransaction Transaction( LOCTEXT("ConvertToSequenceEvaluator", "Convert to Single Frame Animation") );
 				// convert to sequence player
 				UEdGraph* TargetGraph = OldNode->GetGraph();
+				TargetGraph->Modify();
+				OldNode->Modify();
+
 				// create new player
 				FGraphNodeCreator<UAnimGraphNode_BlendSpacePlayer> NodeCreator(*TargetGraph);
 				UAnimGraphNode_BlendSpacePlayer* NewNode = NodeCreator.CreateNode();
@@ -874,36 +923,64 @@ void FAnimationBlueprintEditor::OnConvertToBlendSpacePlayer()
 				// get default data from old node to new node
 				FEdGraphUtilities::CopyCommonState(OldNode, NewNode);
 
-				UEdGraphPin* OldPosePin = OldNode->FindPin(TEXT("X"));
-				UEdGraphPin* NewPosePin = NewNode->FindPin(TEXT("X"));
+				CopyPinData(OldNode, NewNode, TEXT("X"));
+				CopyPinData(OldNode, NewNode, TEXT("Y"));
+				CopyPinData(OldNode, NewNode, TEXT("Pose"));
 
-				if (ensure(OldPosePin && NewPosePin))
-				{
-					NewPosePin->MovePersistentDataFromOldPin(*OldPosePin);
-				}
+				// remove from selection and from graph
+				NodeIter.RemoveCurrent();
+				TargetGraph->RemoveNode(OldNode);
+			}
+		}
 
-				OldPosePin = OldNode->FindPin(TEXT("Y"));
-				NewPosePin = NewNode->FindPin(TEXT("Y"));
+		// @todo fixme: below code doesn't work
+		// because of SetAndCenterObject kicks in after new node is added
+		// will need to disable that first
+		TSharedPtr<SGraphEditor> FocusedGraphEd = FocusedGraphEdPtr.Pin();
+		// Update the graph so that the node will be refreshed
+		FocusedGraphEd->NotifyGraphChanged();
+		// It's possible to leave invalid objects in the selection set if they get GC'd, so clear it out
+		FocusedGraphEd->ClearSelectionSet();
 
-				if (ensure(OldPosePin && NewPosePin))
-				{
-					NewPosePin->MovePersistentDataFromOldPin(*OldPosePin);
-				}
+		FBlueprintEditorUtils::MarkBlueprintAsStructurallyModified(GetAnimBlueprint());
+	}
+}
 
+void FAnimationBlueprintEditor::OnConvertToBlendSpaceGraph()
+{
+	FGraphPanelSelectionSet SelectedNodes = GetSelectedNodes();
+	if (SelectedNodes.Num() > 0)
+	{
+		const FScopedTransaction Transaction( LOCTEXT("ConvertToblendSpaceGraph", "Convert to Blend Space Graph") );
 
-				OldPosePin = OldNode->FindPin(TEXT("Pose"));
-				NewPosePin = NewNode->FindPin(TEXT("Pose"));
+		for (auto NodeIter = SelectedNodes.CreateIterator(); NodeIter; ++NodeIter)
+		{
+			UAnimGraphNode_BlendSpacePlayer* OldNode = Cast<UAnimGraphNode_BlendSpacePlayer>(*NodeIter);
 
-				if (ensure(OldPosePin && NewPosePin))
-				{
-					NewPosePin->MovePersistentDataFromOldPin(*OldPosePin);
-				}
+			// see if sequence player
+			if (OldNode && OldNode->Node.BlendSpace)
+			{
+				// convert to sequence player
+				UEdGraph* TargetGraph = OldNode->GetGraph();
+				TargetGraph->Modify();
+				OldNode->Modify();
+
+				// create new player
+				FGraphNodeCreator<UAnimGraphNode_BlendSpaceGraph> NodeCreator(*TargetGraph);
+				UAnimGraphNode_BlendSpaceGraph* NewNode = NodeCreator.CreateNode();
+				NewNode->SetupFromAsset(OldNode->Node.BlendSpace, false);
+				NodeCreator.Finalize();
+
+				// get default data from old node to new node
+				FEdGraphUtilities::CopyCommonState(OldNode, NewNode);
+
+				CopyPinData(OldNode, NewNode, TEXT("X"));
+				CopyPinData(OldNode, NewNode, TEXT("Y"));
+				CopyPinData(OldNode, NewNode, TEXT("Pose"));
 
 				// remove from selection and from graph
 				NodeIter.RemoveCurrent();
 				TargetGraph->RemoveNode(OldNode);
-
-				NewNode->Modify();
 			}
 		}
 
@@ -925,6 +1002,8 @@ void FAnimationBlueprintEditor::OnConvertToPoseBlender()
 	FGraphPanelSelectionSet SelectedNodes = GetSelectedNodes();
 	if (SelectedNodes.Num() > 0)
 	{
+		const FScopedTransaction Transaction( LOCTEXT("ConvertToPoseBlender", "Convert to Pose Blender") );
+
 		for (auto NodeIter = SelectedNodes.CreateIterator(); NodeIter; ++NodeIter)
 		{
 			UAnimGraphNode_PoseByName* OldNode = Cast<UAnimGraphNode_PoseByName>(*NodeIter);
@@ -934,6 +1013,9 @@ void FAnimationBlueprintEditor::OnConvertToPoseBlender()
 			{
 				// convert to sequence player
 				UEdGraph* TargetGraph = OldNode->GetGraph();
+				TargetGraph->Modify();
+				OldNode->Modify();
+
 				// create new player
 				FGraphNodeCreator<UAnimGraphNode_PoseBlendNode> NodeCreator(*TargetGraph);
 				UAnimGraphNode_PoseBlendNode* NewNode = NodeCreator.CreateNode();
@@ -943,19 +1025,11 @@ void FAnimationBlueprintEditor::OnConvertToPoseBlender()
 				// get default data from old node to new node
 				FEdGraphUtilities::CopyCommonState(OldNode, NewNode);
 
-				UEdGraphPin* OldPosePin = OldNode->FindPin(TEXT("Pose"));
-				UEdGraphPin* NewPosePin = NewNode->FindPin(TEXT("Pose"));
-
-				if (ensure(OldPosePin && NewPosePin))
-				{
-					NewPosePin->MovePersistentDataFromOldPin(*OldPosePin);
-				}
+				CopyPinData(OldNode, NewNode, TEXT("Pose"));
 
 				// remove from selection and from graph
 				NodeIter.RemoveCurrent();
 				TargetGraph->RemoveNode(OldNode);
-
-				NewNode->Modify();
 			}
 		}
 
@@ -978,6 +1052,8 @@ void FAnimationBlueprintEditor::OnConvertToPoseByName()
 	FGraphPanelSelectionSet SelectedNodes = GetSelectedNodes();
 	if (SelectedNodes.Num() > 0)
 	{
+		const FScopedTransaction Transaction( LOCTEXT("ConvertToPoseByName", "Convert to Pose By Name") );
+
 		for (auto NodeIter = SelectedNodes.CreateIterator(); NodeIter; ++NodeIter)
 		{
 			UAnimGraphNode_PoseBlendNode* OldNode = Cast<UAnimGraphNode_PoseBlendNode>(*NodeIter);
@@ -985,9 +1061,11 @@ void FAnimationBlueprintEditor::OnConvertToPoseByName()
 			// see if sequence player
 			if (OldNode && OldNode->Node.PoseAsset)
 			{
-				//const FScopedTransaction Transaction( LOCTEXT("ConvertToSequenceEvaluator", "Convert to Single Frame Animation") );
 				// convert to sequence player
 				UEdGraph* TargetGraph = OldNode->GetGraph();
+				TargetGraph->Modify();
+				OldNode->Modify();
+
 				// create new player
 				FGraphNodeCreator<UAnimGraphNode_PoseByName> NodeCreator(*TargetGraph);
 				UAnimGraphNode_PoseByName* NewNode = NodeCreator.CreateNode();
@@ -997,19 +1075,11 @@ void FAnimationBlueprintEditor::OnConvertToPoseByName()
 				// get default data from old node to new node
 				FEdGraphUtilities::CopyCommonState(OldNode, NewNode);
 
-				UEdGraphPin* OldPosePin = OldNode->FindPin(TEXT("Pose"));
-				UEdGraphPin* NewPosePin = NewNode->FindPin(TEXT("Pose"));
-
-				if (ensure(OldPosePin && NewPosePin))
-				{
-					NewPosePin->MovePersistentDataFromOldPin(*OldPosePin);
-				}
+				CopyPinData(OldNode, NewNode, TEXT("Pose"));
 
 				// remove from selection and from graph
 				NodeIter.RemoveCurrent();
 				TargetGraph->RemoveNode(OldNode);
-
-				NewNode->Modify();
 			}
 		}
 
@@ -1033,6 +1103,8 @@ void FAnimationBlueprintEditor::OnConvertToAimOffsetLookAt()
 
 	if (SelectedNodes.Num() > 0)
 	{
+		const FScopedTransaction Transaction( LOCTEXT("ConvertToAimOffsetLookAt", "Convert to Aim Offset LookAt") );
+
 		for (auto NodeIter = SelectedNodes.CreateIterator(); NodeIter; ++NodeIter)
 		{
 			UAnimGraphNode_RotationOffsetBlendSpace* OldNode = Cast<UAnimGraphNode_RotationOffsetBlendSpace>(*NodeIter);
@@ -1040,10 +1112,11 @@ void FAnimationBlueprintEditor::OnConvertToAimOffsetLookAt()
 			// see if sequence player
 			if (OldNode && OldNode->Node.BlendSpace)
 			{
-				//const FScopedTransaction Transaction( LOCTEXT("ConvertToSequenceEvaluator", "Convert to Single Frame Animation") );
-
 				// convert to sequence evaluator
 				UEdGraph* TargetGraph = OldNode->GetGraph();
+				TargetGraph->Modify();
+				OldNode->Modify();
+
 				// create new evaluator
 				FGraphNodeCreator<UAnimGraphNode_AimOffsetLookAt> NodeCreator(*TargetGraph);
 				UAnimGraphNode_AimOffsetLookAt* NewNode = NodeCreator.CreateNode();
@@ -1053,27 +1126,15 @@ void FAnimationBlueprintEditor::OnConvertToAimOffsetLookAt()
 				// get default data from old node to new node
 				FEdGraphUtilities::CopyCommonState(OldNode, NewNode);
 
-				UEdGraphPin* OldPosePin = OldNode->FindPin(TEXT("Pose"));
-				UEdGraphPin* NewPosePin = NewNode->FindPin(TEXT("Pose"));
-
-				if (ensure(OldPosePin && NewPosePin))
-				{
-					NewPosePin->MovePersistentDataFromOldPin(*OldPosePin);
-				}
-
-				OldPosePin = OldNode->FindPin(TEXT("BasePose"));
-				NewPosePin = NewNode->FindPin(TEXT("BasePose"));
-
-				if (ensure(OldPosePin && NewPosePin))
-				{
-					NewPosePin->MovePersistentDataFromOldPin(*OldPosePin);
-				}
+				CopyPinData(OldNode, NewNode, TEXT("X"));
+				CopyPinData(OldNode, NewNode, TEXT("Y"));
+				CopyPinData(OldNode, NewNode, TEXT("Alpha"));
+				CopyPinData(OldNode, NewNode, TEXT("Pose"));
+				CopyPinData(OldNode, NewNode, TEXT("BasePose"));
 
 				// remove from selection and from graph
 				NodeIter.RemoveCurrent();
 				TargetGraph->RemoveNode(OldNode);
-
-				NewNode->Modify();
 			}
 		}
 
@@ -1096,6 +1157,8 @@ void FAnimationBlueprintEditor::OnConvertToAimOffsetSimple()
 	FGraphPanelSelectionSet SelectedNodes = GetSelectedNodes();
 	if (SelectedNodes.Num() > 0)
 	{
+		const FScopedTransaction Transaction( LOCTEXT("ConvertToSimpleAimOffset", "Convert to Simple Aim Offset") );
+
 		for (auto NodeIter = SelectedNodes.CreateIterator(); NodeIter; ++NodeIter)
 		{
 			UAnimGraphNode_AimOffsetLookAt* OldNode = Cast<UAnimGraphNode_AimOffsetLookAt>(*NodeIter);
@@ -1103,9 +1166,11 @@ void FAnimationBlueprintEditor::OnConvertToAimOffsetSimple()
 			// see if sequence player
 			if (OldNode && OldNode->Node.BlendSpace)
 			{
-				//const FScopedTransaction Transaction( LOCTEXT("ConvertToSequenceEvaluator", "Convert to Single Frame Animation") );
 				// convert to sequence player
 				UEdGraph* TargetGraph = OldNode->GetGraph();
+				TargetGraph->Modify();
+				OldNode->Modify();
+
 				// create new player
 				FGraphNodeCreator<UAnimGraphNode_RotationOffsetBlendSpace> NodeCreator(*TargetGraph);
 				UAnimGraphNode_RotationOffsetBlendSpace* NewNode = NodeCreator.CreateNode();
@@ -1115,27 +1180,67 @@ void FAnimationBlueprintEditor::OnConvertToAimOffsetSimple()
 				// get default data from old node to new node
 				FEdGraphUtilities::CopyCommonState(OldNode, NewNode);
 
-				UEdGraphPin* OldPosePin = OldNode->FindPin(TEXT("Pose"));
-				UEdGraphPin* NewPosePin = NewNode->FindPin(TEXT("Pose"));
+				CopyPinData(OldNode, NewNode, TEXT("X"));
+				CopyPinData(OldNode, NewNode, TEXT("Y"));
+				CopyPinData(OldNode, NewNode, TEXT("Pose"));
+				CopyPinData(OldNode, NewNode, TEXT("BasePose"));
 
-				if (ensure(OldPosePin && NewPosePin))
-				{
-					NewPosePin->MovePersistentDataFromOldPin(*OldPosePin);
-				}
+				// remove from selection and from graph
+				NodeIter.RemoveCurrent();
+				TargetGraph->RemoveNode(OldNode);
+			}
+		}
+
+		// @todo fixme: below code doesn't work
+		// because of SetAndCenterObject kicks in after new node is added
+		// will need to disable that first
+		TSharedPtr<SGraphEditor> FocusedGraphEd = FocusedGraphEdPtr.Pin();
+		// Update the graph so that the node will be refreshed
+		FocusedGraphEd->NotifyGraphChanged();
+		// It's possible to leave invalid objects in the selection set if they get GC'd, so clear it out
+		FocusedGraphEd->ClearSelectionSet();
+
+		FBlueprintEditorUtils::MarkBlueprintAsStructurallyModified(GetAnimBlueprint());
+	}
+}
 
-				OldPosePin = OldNode->FindPin(TEXT("BasePose"));
-				NewPosePin = NewNode->FindPin(TEXT("BasePose"));
+void FAnimationBlueprintEditor::OnConvertToAimOffsetGraph()
+{
+	FGraphPanelSelectionSet SelectedNodes = GetSelectedNodes();
+	if (SelectedNodes.Num() > 0)
+	{
+		const FScopedTransaction Transaction( LOCTEXT("ConvertToAimOffsetGraph", "Convert to Aim Offset Graph") );
 
-				if (ensure(OldPosePin && NewPosePin))
-				{
-					NewPosePin->MovePersistentDataFromOldPin(*OldPosePin);
-				}
+		for (auto NodeIter = SelectedNodes.CreateIterator(); NodeIter; ++NodeIter)
+		{
+			UAnimGraphNode_RotationOffsetBlendSpace* OldNode = Cast<UAnimGraphNode_RotationOffsetBlendSpace>(*NodeIter);
+
+			// see if sequence player
+			if (OldNode && OldNode->Node.BlendSpace)
+			{
+				// convert to sequence player
+				UEdGraph* TargetGraph = OldNode->GetGraph();
+				TargetGraph->Modify();
+				OldNode->Modify();
+
+				// create new player
+				FGraphNodeCreator<UAnimGraphNode_RotationOffsetBlendSpaceGraph> NodeCreator(*TargetGraph);
+				UAnimGraphNode_RotationOffsetBlendSpaceGraph* NewNode = NodeCreator.CreateNode();
+				NewNode->SetupFromAsset(OldNode->Node.BlendSpace, false);
+				NodeCreator.Finalize();
+
+				// get default data from old node to new node
+				FEdGraphUtilities::CopyCommonState(OldNode, NewNode);
+
+				CopyPinData(OldNode, NewNode, TEXT("X"));
+				CopyPinData(OldNode, NewNode, TEXT("Y"));
+				CopyPinData(OldNode, NewNode, TEXT("Alpha"));
+				CopyPinData(OldNode, NewNode, TEXT("Pose"));
+				CopyPinData(OldNode, NewNode, TEXT("BasePose"));
 
 				// remove from selection and from graph
 				NodeIter.RemoveCurrent();
 				TargetGraph->RemoveNode(OldNode);
-
-				NewNode->Modify();
 			}
 		}
 
@@ -1546,6 +1651,19 @@ void FAnimationBlueprintEditor::CreateEditorModeManager()
 	PRAGMA_ENABLE_DEPRECATION_WARNINGS
 }
 
+void FAnimationBlueprintEditor::JumpToHyperlink(const UObject* ObjectReference, bool bRequestRename)
+{
+	if(const UBlendSpaceGraph* BlendSpaceGraph = Cast<UBlendSpaceGraph>(ObjectReference))
+	{
+		TSharedRef<FTabPayload_BlendSpaceGraph> Payload = FTabPayload_BlendSpaceGraph::Make(BlendSpaceGraph);
+		DocumentManager->OpenDocument(Payload, FDocumentTracker::OpenNewDocument);
+	}
+	else
+	{
+		FBlueprintEditor::JumpToHyperlink(ObjectReference, bRequestRename);
+	}
+}
+
 TSharedRef<IPersonaPreviewScene> FAnimationBlueprintEditor::GetPreviewScene() const
 { 
 	return PersonaToolkit->GetPreviewScene(); 
diff --git a/Engine/Source/Editor/AnimationBlueprintEditor/Private/AnimationBlueprintEditor.h b/Engine/Source/Editor/AnimationBlueprintEditor/Private/AnimationBlueprintEditor.h
--- a/Engine/Source/Editor/AnimationBlueprintEditor/Private/AnimationBlueprintEditor.h
+++ b/Engine/Source/Editor/AnimationBlueprintEditor/Private/AnimationBlueprintEditor.h
@@ -162,6 +162,10 @@ public:
 	virtual TStatId GetStatId() const override;
 	//~ End FTickableEditorObject Interface
 
+	//~ Begin FBlueprintEditor Interface
+	virtual void JumpToHyperlink(const UObject* ObjectReference, bool bRequestRename) override;
+	//~ End FBlueprintEditor Interface
+
 	TSharedRef<SWidget> GetPreviewEditor() { return PreviewEditor.ToSharedRef(); }
 	/** Refresh Preview Instance Track Curves **/
 	void RefreshPreviewInstanceTrackCurves();
@@ -171,28 +175,17 @@ public:
 	/** Get the skeleton tree this Persona editor is hosting */
 	TSharedRef<class ISkeletonTree> GetSkeletonTree() const { return SkeletonTree.ToSharedRef(); }
 
+	/** Make this available to allow us to create title bar widgets for other container types - e.g. blendspaces */
+	using FBlueprintEditor::CreateGraphTitleBarWidget;
+
 protected:
 	//~ Begin FBlueprintEditor Interface
 	//virtual void CreateDefaultToolbar() override;
 	virtual void CreateDefaultCommands() override;
 	virtual void OnCreateGraphEditorCommands(TSharedPtr<FUICommandList> GraphEditorCommandsList);
-	virtual bool CanSelectBone() const override { return true; }
-	virtual void OnAddPosePin() override;
-	virtual bool CanAddPosePin() const override;
-	virtual void OnRemovePosePin() override;
-	virtual bool CanRemovePosePin() const override;
 	virtual void OnGraphEditorFocused(const TSharedRef<class SGraphEditor>& InGraphEditor) override;
 	virtual void OnGraphEditorBackgrounded(const TSharedRef<SGraphEditor>& InGraphEditor) override;
-	virtual void OnConvertToSequenceEvaluator() override;
-	virtual void OnConvertToSequencePlayer() override;
-	virtual void OnConvertToBlendSpaceEvaluator() override;
-	virtual void OnConvertToBlendSpacePlayer() override;
-	virtual void OnConvertToPoseBlender() override;
-	virtual void OnConvertToPoseByName() override;
-	virtual void OnConvertToAimOffsetLookAt() override;
-	virtual void OnConvertToAimOffsetSimple() override;
 	virtual bool IsInAScriptingMode() const override { return true; }
-	virtual void OnOpenRelatedAsset() override;
 	virtual void GetCustomDebugObjects(TArray<FCustomDebugObject>& DebugList) const override;
 	virtual void CreateDefaultTabContents(const TArray<UBlueprint*>& InBlueprints) override;
 	virtual FGraphAppearanceInfo GetGraphAppearance(class UEdGraph* InGraph) const override;
@@ -312,6 +305,27 @@ private:
 	/** Chooses a suitable pose watch colour automatically - i.e. one that isn't already in use (if possible) */
 	FColor ChoosePoseWatchColor() const;
 
+	// Pose pin UI handlers
+	void OnAddPosePin();
+	bool CanAddPosePin() const;
+	void OnRemovePosePin();
+	bool CanRemovePosePin() const;
+
+	// Node conversion functions
+	void OnConvertToSequenceEvaluator();
+	void OnConvertToSequencePlayer();
+	void OnConvertToBlendSpaceEvaluator();
+	void OnConvertToBlendSpacePlayer();
+	void OnConvertToBlendSpaceGraph();
+	void OnConvertToPoseBlender();
+	void OnConvertToPoseByName();
+	void OnConvertToAimOffsetLookAt();
+	void OnConvertToAimOffsetSimple();
+	void OnConvertToAimOffsetGraph();
+	
+	// Opens the associated asset of the selected nodes
+	void OnOpenRelatedAsset();
+
 	/** The extender to pass to the level editor to extend it's window menu */
 	TSharedPtr<FExtender> MenuExtender;
 
