Non-unity/non-PCH CIS fixes

diff --git a/Engine/Source/Editor/AnimGraph/Public/BlendSpaceGraph.h b/Engine/Source/Editor/AnimGraph/Public/BlendSpaceGraph.h
--- a/Engine/Source/Editor/AnimGraph/Public/BlendSpaceGraph.h
+++ b/Engine/Source/Editor/AnimGraph/Public/BlendSpaceGraph.h
@@ -3,6 +3,7 @@
 #pragma once
 
 #include "CoreMinimal.h"
+#include "EdGraph/EdGraph.h"
 #include "BlendSpaceGraph.generated.h"
 
 class UBlendSpaceBase;
diff --git a/Engine/Source/Editor/AnimationBlueprintEditor/Private/TabPayload_BlendSpaceGraph.h b/Engine/Source/Editor/AnimationBlueprintEditor/Private/TabPayload_BlendSpaceGraph.h
--- a/Engine/Source/Editor/AnimationBlueprintEditor/Private/TabPayload_BlendSpaceGraph.h
+++ b/Engine/Source/Editor/AnimationBlueprintEditor/Private/TabPayload_BlendSpaceGraph.h
@@ -4,6 +4,7 @@
 
 #include "CoreMinimal.h"
 #include "WorkflowOrientedApp/WorkflowUObjectDocuments.h"
+#include "BlendSpaceGraph.h"
 
 // Tab payload used to generate blendspace documents
 struct FTabPayload_BlendSpaceGraph : public FTabPayload
