Non-unity/non-PCH CIS fix

diff --git a/Engine/Source/Editor/AnimationBlueprintEditor/Private/AnimationBlueprintEditor.cpp b/Engine/Source/Editor/AnimationBlueprintEditor/Private/AnimationBlueprintEditor.cpp
--- a/Engine/Source/Editor/AnimationBlueprintEditor/Private/AnimationBlueprintEditor.cpp
+++ b/Engine/Source/Editor/AnimationBlueprintEditor/Private/AnimationBlueprintEditor.cpp
@@ -85,6 +85,7 @@
 #include "TabPayload_BlendSpaceGraph.h"
 #include "AnimGraphNode_RotationOffsetBlendSpaceGraph.h"
 #include "AnimGraphNode_BlendSpaceGraph.h"
+#include "ScopedTransaction.h"
 
 #define LOCTEXT_NAMESPACE "AnimationBlueprintEditor"
 
