Keep blendspace parameters around when we have a blendspace graph document focused

diff --git a/Engine/Source/Editor/AnimGraph/Public/BlendSpaceGraph.h b/Engine/Source/Editor/AnimGraph/Public/BlendSpaceGraph.h
--- a/Engine/Source/Editor/AnimGraph/Public/BlendSpaceGraph.h
+++ b/Engine/Source/Editor/AnimGraph/Public/BlendSpaceGraph.h
@@ -20,6 +20,6 @@ private:
 
 public:
 	// Blendspace that we wrap
-	UPROPERTY()
+	UPROPERTY(VisibleAnywhere, Category="BlendSpace", Instanced, NoClear, meta = (ShowOnlyInnerProperties))
 	TObjectPtr<UBlendSpace> BlendSpace;
 };
\ No newline at end of file
diff --git a/Engine/Source/Editor/AnimationBlueprintEditor/Private/BlendSpaceDocumentTabFactory.cpp b/Engine/Source/Editor/AnimationBlueprintEditor/Private/BlendSpaceDocumentTabFactory.cpp
--- a/Engine/Source/Editor/AnimationBlueprintEditor/Private/BlendSpaceDocumentTabFactory.cpp
+++ b/Engine/Source/Editor/AnimationBlueprintEditor/Private/BlendSpaceDocumentTabFactory.cpp
@@ -15,11 +15,34 @@
 #include "Kismet2/BlueprintEditorUtils.h"
 #include "SGraphPreviewer.h"
 #include "AnimNodes/AnimNode_BlendSpaceGraphBase.h"
+#include "Widgets/Docking/SDockTab.h"
 
 static const FName BlendSpaceEditorID("BlendSpaceEditor");
 
 #define LOCTEXT_NAMESPACE "FBlendSpaceDocumentTabFactory"
 
+// Simple wrapper widget used to hold a reference to the graph document
+class SBlendSpaceDocumentTab : public SCompoundWidget
+{	
+	SLATE_BEGIN_ARGS(SBlendSpaceDocumentTab) {}
+
+	SLATE_DEFAULT_SLOT(FArguments, Content)
+	
+	SLATE_END_ARGS()
+
+	void Construct(const FArguments& InArgs, UBlendSpaceGraph* InDocument)
+	{
+		Document = InDocument;
+		
+		ChildSlot
+		[
+			InArgs._Content.Widget
+		];
+	}
+
+	TWeakObjectPtr<UBlendSpaceGraph> Document;
+};
+
 FBlendSpaceDocumentTabFactory::FBlendSpaceDocumentTabFactory(TSharedPtr<FAnimationBlueprintEditor> InBlueprintEditorPtr)
 	: FDocumentTabFactory(BlendSpaceEditorID, InBlueprintEditorPtr)
 	, BlueprintEditorPtr(InBlueprintEditorPtr)
@@ -188,17 +211,20 @@ TSharedRef<SWidget> FBlendSpaceDocumentTabFactory::CreateTabBody(const FWorkflow
 
 	Args.StatusBarName = TEXT("AssetEditor.AnimationBlueprintEditor.MainMenu");
 
-	return 
-		SNew(SVerticalBox)
-		+SVerticalBox::Slot()
-		.AutoHeight()
-		[
-			BlueprintEditorPtr.Pin()->CreateGraphTitleBarWidget(Info.TabInfo.ToSharedRef(), DocumentID)
-		]
-		+SVerticalBox::Slot()
-		.FillHeight(1.0f)
+	return
+		SNew(SBlendSpaceDocumentTab, DocumentID)
 		[
-			PersonaModule.CreateBlendSpaceEditWidget(DocumentID->BlendSpace, Args)
+			SNew(SVerticalBox)
+			+SVerticalBox::Slot()
+			.AutoHeight()
+			[
+				BlueprintEditorPtr.Pin()->CreateGraphTitleBarWidget(Info.TabInfo.ToSharedRef(), DocumentID)
+			]
+			+SVerticalBox::Slot()
+			.FillHeight(1.0f)
+			[
+				PersonaModule.CreateBlendSpaceEditWidget(DocumentID->BlendSpace, Args)
+			]
 		];
 }
 
@@ -242,5 +268,28 @@ TAttribute<FText> FBlendSpaceDocumentTabFactory::ConstructTabName(const FWorkflo
 		});
 }
 
+void FBlendSpaceDocumentTabFactory::OnTabActivated(TSharedPtr<SDockTab> Tab) const
+{
+	TSharedRef<SBlendSpaceDocumentTab> DocumentWidget = StaticCastSharedRef<SBlendSpaceDocumentTab>(Tab->GetContent());
+	if(UBlendSpaceGraph* Document = DocumentWidget->Document.Get())
+	{
+		if(TSharedPtr<FAnimationBlueprintEditor> BlueprintEditor = BlueprintEditorPtr.Pin())
+		{
+			BlueprintEditor->SetDetailObject(Document);
+		}
+	}
+}
+
+void FBlendSpaceDocumentTabFactory::OnTabForegrounded(TSharedPtr<SDockTab> Tab) const
+{
+	TSharedRef<SBlendSpaceDocumentTab> DocumentWidget = StaticCastSharedRef<SBlendSpaceDocumentTab>(Tab->GetContent());
+	if(UBlendSpaceGraph* Document = DocumentWidget->Document.Get())
+	{
+		if(TSharedPtr<FAnimationBlueprintEditor> BlueprintEditor = BlueprintEditorPtr.Pin())
+		{
+			BlueprintEditor->SetDetailObject(Document);
+		}
+	}
+}
 
 #undef LOCTEXT_NAMESPACE
\ No newline at end of file
diff --git a/Engine/Source/Editor/AnimationBlueprintEditor/Private/BlendSpaceDocumentTabFactory.h b/Engine/Source/Editor/AnimationBlueprintEditor/Private/BlendSpaceDocumentTabFactory.h
--- a/Engine/Source/Editor/AnimationBlueprintEditor/Private/BlendSpaceDocumentTabFactory.h
+++ b/Engine/Source/Editor/AnimationBlueprintEditor/Private/BlendSpaceDocumentTabFactory.h
@@ -19,7 +19,9 @@ protected:
 	virtual bool IsPayloadSupported(TSharedRef<FTabPayload> Payload) const override;
 	virtual bool IsPayloadValid(TSharedRef<FTabPayload> Payload) const override;
 	virtual TAttribute<FText> ConstructTabName(const FWorkflowTabSpawnInfo& Info) const override;
-
+	virtual void OnTabActivated(TSharedPtr<SDockTab> Tab) const override;
+	virtual void OnTabForegrounded(TSharedPtr<SDockTab> Tab) const override;
+	
 protected:
 	TWeakPtr<FAnimationBlueprintEditor> BlueprintEditorPtr;
 };
