Added the option to remove the "ANIMATION" overlay for anim graphs

diff --git a/Engine/Source/Editor/AnimationBlueprintEditor/Private/AnimationBlueprintEditor.cpp b/Engine/Source/Editor/AnimationBlueprintEditor/Private/AnimationBlueprintEditor.cpp
--- a/Engine/Source/Editor/AnimationBlueprintEditor/Private/AnimationBlueprintEditor.cpp
+++ b/Engine/Source/Editor/AnimationBlueprintEditor/Private/AnimationBlueprintEditor.cpp
@@ -86,6 +86,7 @@
 #include "AnimGraphNode_RotationOffsetBlendSpaceGraph.h"
 #include "AnimGraphNode_BlendSpaceGraph.h"
 #include "ScopedTransaction.h"
+#include "Preferences/PersonaOptions.h"
 
 #define LOCTEXT_NAMESPACE "AnimationBlueprintEditor"
 
@@ -1465,7 +1466,7 @@ FGraphAppearanceInfo FAnimationBlueprintEditor::GetGraphAppearance(UEdGraph* InG
 
 	if ( GetBlueprintObj()->IsA(UAnimBlueprint::StaticClass()) )
 	{
-		AppearanceInfo.CornerText = LOCTEXT("AppearanceCornerText_Animation", "ANIMATION");
+		AppearanceInfo.CornerText = (GetDefault<UAnimationBlueprintEditorSettings>()->bShowGraphCornerText) ? LOCTEXT("AppearanceCornerText_Animation", "ANIMATION") : FText::GetEmpty();
 	}
 
 	return AppearanceInfo;
diff --git a/Engine/Source/Editor/AnimationBlueprintEditor/Public/AnimationBlueprintEditorSettings.h b/Engine/Source/Editor/AnimationBlueprintEditor/Public/AnimationBlueprintEditorSettings.h
--- a/Engine/Source/Editor/AnimationBlueprintEditor/Public/AnimationBlueprintEditorSettings.h
+++ b/Engine/Source/Editor/AnimationBlueprintEditor/Public/AnimationBlueprintEditorSettings.h
@@ -20,6 +20,10 @@ public:
 	UPROPERTY(EditAnywhere, config, Category = "Debugging")
 	bool bPoseWatchSelectedNodes = false;
 
+	/** Whether to display the corner text in an animation graph. Changing this only affects newly opened graphs. */
+	UPROPERTY(EditAnywhere, config, Category = "Graphs")
+	bool bShowGraphCornerText = true;
+
 public:
 	DECLARE_MULTICAST_DELEGATE_TwoParams(FOnUpdateSettingsMulticaster, const UAnimationBlueprintEditorSettings*, EPropertyChangeType::Type);
 	FOnUpdateSettingsMulticaster OnSettingsChange;
