Fixed an issue where the AnimationInstance ActiveAnimNotifyEventReference was not kept in sync with ActiveAnimNotifyState

diff --git a/Engine/Source/Editor/AnimationBlueprintEditor/Private/AnimationBlueprintEditor.cpp b/Engine/Source/Editor/AnimationBlueprintEditor/Private/AnimationBlueprintEditor.cpp
--- a/Engine/Source/Editor/AnimationBlueprintEditor/Private/AnimationBlueprintEditor.cpp
+++ b/Engine/Source/Editor/AnimationBlueprintEditor/Private/AnimationBlueprintEditor.cpp
@@ -1583,7 +1583,9 @@ void FAnimationBlueprintEditor::OnBlueprintPreCompile(UBlueprint* BlueprintToCom
 				if(Event.NotifyStateClass->GetClass() == BlueprintToCompile->GeneratedClass)
 				{
 					Event.NotifyStateClass->NotifyEnd(SkelMeshComp, Cast<UAnimSequenceBase>(Event.NotifyStateClass->GetOuter()), EventReference);
+					check(Instance->ActiveAnimNotifyState.Num() == Instance->ActiveAnimNotifyEventReference.Num());
 					Instance->ActiveAnimNotifyState.RemoveAt(Idx);
+					Instance->ActiveAnimNotifyEventReference.RemoveAt(Idx);
 				}
 			}
 		}
diff --git a/Engine/Source/Runtime/Engine/Private/Animation/AnimInstance.cpp b/Engine/Source/Runtime/Engine/Private/Animation/AnimInstance.cpp
--- a/Engine/Source/Runtime/Engine/Private/Animation/AnimInstance.cpp
+++ b/Engine/Source/Runtime/Engine/Private/Animation/AnimInstance.cpp
@@ -1795,7 +1795,9 @@ void UAnimInstance::TriggerMontageEndedEvent(const FQueuedMontageEndedEvent& Mon
 
 				if (ActiveAnimNotifyState.IsValidIndex(Index))
 				{
+					check(ActiveAnimNotifyState.Num() == ActiveAnimNotifyEventReference.Num());
 					ActiveAnimNotifyState.RemoveAtSwap(Index);
+					ActiveAnimNotifyEventReference.RemoveAtSwap(Index);
 				}
 				else
 				{
