Navigation into/out of Blendspace graphs

diff --git a/Engine/Source/Editor/AnimationBlueprintEditor/Private/BlendSpaceDocumentTabFactory.cpp b/Engine/Source/Editor/AnimationBlueprintEditor/Private/BlendSpaceDocumentTabFactory.cpp
--- a/Engine/Source/Editor/AnimationBlueprintEditor/Private/BlendSpaceDocumentTabFactory.cpp
+++ b/Engine/Source/Editor/AnimationBlueprintEditor/Private/BlendSpaceDocumentTabFactory.cpp
@@ -16,6 +16,8 @@
 #include "SGraphPreviewer.h"
 #include "AnimNodes/AnimNode_BlendSpaceGraphBase.h"
 #include "Widgets/Docking/SDockTab.h"
+#include "BlueprintEditorSettings.h"
+#include "Framework/MultiBox/MultiBoxBuilder.h"
 
 static const FName BlendSpaceEditorID("BlendSpaceEditor");
 
@@ -30,9 +32,9 @@ class SBlendSpaceDocumentTab : public SCompoundWidget
 	
 	SLATE_END_ARGS()
 
-	void Construct(const FArguments& InArgs, UBlendSpaceGraph* InDocument)
+	void Construct(const FArguments& InArgs, UBlendSpaceGraph* InBlendSpaceGraph)
 	{
-		Document = InDocument;
+		BlendSpaceGraph = InBlendSpaceGraph;
 		
 		ChildSlot
 		[
@@ -40,7 +42,7 @@ class SBlendSpaceDocumentTab : public SCompoundWidget
 		];
 	}
 
-	TWeakObjectPtr<UBlendSpaceGraph> Document;
+	TWeakObjectPtr<UBlendSpaceGraph> BlendSpaceGraph;
 };
 
 FBlendSpaceDocumentTabFactory::FBlendSpaceDocumentTabFactory(TSharedPtr<FAnimationBlueprintEditor> InBlueprintEditorPtr)
@@ -51,14 +53,82 @@ FBlendSpaceDocumentTabFactory::FBlendSpaceDocumentTabFactory(TSharedPtr<FAnimati
 
 TSharedRef<SWidget> FBlendSpaceDocumentTabFactory::CreateTabBody(const FWorkflowTabSpawnInfo& Info) const
 {
-	UBlendSpaceGraph* DocumentID = FTabPayload_BlendSpaceGraph::GetBlendSpaceGraph(Info.Payload);
+	UBlendSpaceGraph* BlendSpaceGraph = FTabPayload_BlendSpaceGraph::GetBlendSpaceGraph(Info.Payload);
 
 	FPersonaModule& PersonaModule = FModuleManager::LoadModuleChecked<FPersonaModule>("Persona");
 
-	UAnimGraphNode_BlendSpaceGraphBase* BlendSpaceNode = CastChecked<UAnimGraphNode_BlendSpaceGraphBase>(DocumentID->GetOuter());
+	UAnimGraphNode_BlendSpaceGraphBase* BlendSpaceNode = CastChecked<UAnimGraphNode_BlendSpaceGraphBase>(BlendSpaceGraph->GetOuter());
 
 	FBlendSpaceEditorArgs Args;
 
+	Args.OnBlendSpaceCanvasDoubleClicked = FOnBlendSpaceCanvasDoubleClicked::CreateLambda([this, WeakBlendSpaceNode = TWeakObjectPtr<UAnimGraphNode_BlendSpaceGraphBase>(BlendSpaceNode)]()
+	{
+		UBlueprintEditorSettings const* Settings = GetDefault<UBlueprintEditorSettings>();
+		if (Settings->bDoubleClickNavigatesToParent)
+		{
+			if(BlueprintEditorPtr.IsValid() && WeakBlendSpaceNode.Get())
+			{
+				UAnimGraphNode_BlendSpaceGraphBase* BlendSpaceNode = WeakBlendSpaceNode.Get();
+				BlueprintEditorPtr.Pin()->JumpToHyperlink(BlendSpaceNode->GetOuter(), false);
+			}
+		}
+	});
+
+	Args.OnBlendSpaceNavigateUp = FOnBlendSpaceNavigateUp::CreateLambda([this, WeakBlendSpaceNode = TWeakObjectPtr<UAnimGraphNode_BlendSpaceGraphBase>(BlendSpaceNode)]()
+	{
+		if(BlueprintEditorPtr.IsValid() && WeakBlendSpaceNode.Get())
+		{
+			UAnimGraphNode_BlendSpaceGraphBase* BlendSpaceNode = WeakBlendSpaceNode.Get();
+			BlueprintEditorPtr.Pin()->JumpToHyperlink(BlendSpaceNode->GetOuter(), false);
+		}
+	});
+
+	Args.OnBlendSpaceNavigateDown = FOnBlendSpaceNavigateDown::CreateLambda([this, WeakBlendSpaceNode = TWeakObjectPtr<UAnimGraphNode_BlendSpaceGraphBase>(BlendSpaceNode)]()
+	{
+		if(BlueprintEditorPtr.IsValid() && WeakBlendSpaceNode.Get())
+		{
+			UAnimGraphNode_BlendSpaceGraphBase* BlendSpaceNode = WeakBlendSpaceNode.Get();
+			TArrayView<UEdGraph* const> Graphs = BlendSpaceNode->GetGraphs();
+			if (Graphs.Num() > 1)
+			{
+				// Display a child jump list
+				FMenuBuilder MenuBuilder(true, nullptr);
+				MenuBuilder.BeginSection("NavigateToGraph", LOCTEXT("ChildGraphPickerDesc", "Navigate to graph"));
+
+				TArray<const UEdGraph*> SortedGraphs(Graphs);
+				SortedGraphs.Sort([](const UEdGraph& A, const UEdGraph& B) { return FBlueprintEditor::GetGraphDisplayName(&A).CompareToCaseIgnored(FBlueprintEditor::GetGraphDisplayName(&B)) < 0; });
+
+				for (const UEdGraph* Graph : SortedGraphs)
+				{
+					MenuBuilder.AddMenuEntry(
+						BlueprintEditorPtr.Pin()->GetGraphDisplayName(Graph),
+						LOCTEXT("ChildGraphPickerTooltip", "Pick the graph to enter"),
+						FSlateIcon(),
+						FUIAction(
+							FExecuteAction::CreateLambda(
+								[this,Graph]() 
+								{
+									BlueprintEditorPtr.Pin()->OpenDocument(Graph, FDocumentTracker::NavigatingCurrentDocument);
+								}),
+							FCanExecuteAction()));
+				}
+				MenuBuilder.EndSection();
+
+				FSlateApplication::Get().PushMenu( 
+					BlueprintEditorPtr.Pin()->GetToolkitHost()->GetParentWidget(),
+					FWidgetPath(),
+					MenuBuilder.MakeWidget(),
+					FSlateApplication::Get().GetCursorPos(), // summon location
+					FPopupTransitionEffect( FPopupTransitionEffect::TypeInPopup )
+				);
+			}
+			else if (Graphs.Num() == 1)
+			{
+				BlueprintEditorPtr.Pin()->JumpToHyperlink(Graphs[0], false);
+			}
+		}
+	});
+
 	Args.OnBlendSpaceSampleDoubleClicked = FOnBlendSpaceSampleDoubleClicked::CreateLambda([this, WeakBlendSpaceNode = TWeakObjectPtr<UAnimGraphNode_BlendSpaceGraphBase>(BlendSpaceNode)](int32 InSampleIndex)
 	{
 		if(BlueprintEditorPtr.IsValid() && WeakBlendSpaceNode.Get())
@@ -215,27 +285,27 @@ TSharedRef<SWidget> FBlendSpaceDocumentTabFactory::CreateTabBody(const FWorkflow
 	Args.StatusBarName = TEXT("AssetEditor.AnimationBlueprintEditor.MainMenu");
 
 	return
-		SNew(SBlendSpaceDocumentTab, DocumentID)
+		SNew(SBlendSpaceDocumentTab, BlendSpaceGraph)
 		[
 			SNew(SVerticalBox)
 			+SVerticalBox::Slot()
 			.AutoHeight()
 			[
-				BlueprintEditorPtr.Pin()->CreateGraphTitleBarWidget(Info.TabInfo.ToSharedRef(), DocumentID)
+				BlueprintEditorPtr.Pin()->CreateGraphTitleBarWidget(Info.TabInfo.ToSharedRef(), BlendSpaceGraph)
 			]
 			+SVerticalBox::Slot()
 			.FillHeight(1.0f)
 			[
-				PersonaModule.CreateBlendSpaceEditWidget(DocumentID->BlendSpace, Args)
+				PersonaModule.CreateBlendSpaceEditWidget(BlendSpaceGraph->BlendSpace, Args)
 			]
 		];
 }
 
 const FSlateBrush* FBlendSpaceDocumentTabFactory::GetTabIcon(const FWorkflowTabSpawnInfo& Info) const
 {
-	UBlendSpaceGraph* DocumentID = FTabPayload_BlendSpaceGraph::GetBlendSpaceGraph(Info.Payload);
+	UBlendSpaceGraph* BlendSpaceGraph = FTabPayload_BlendSpaceGraph::GetBlendSpaceGraph(Info.Payload);
 
-	if (UBlendSpace1D* BlendSpace1D = Cast<UBlendSpace1D>(DocumentID->BlendSpace))
+	if (UBlendSpace1D* BlendSpace1D = Cast<UBlendSpace1D>(BlendSpaceGraph->BlendSpace))
 	{
 		return FEditorStyle::GetBrush("ClassIcon.BlendSpace1D");
 	}
@@ -263,9 +333,9 @@ TAttribute<FText> FBlendSpaceDocumentTabFactory::ConstructTabName(const FWorkflo
 {
 	check(Info.Payload.IsValid());
 
-	UBlendSpaceGraph* DocumentID = FTabPayload_BlendSpaceGraph::GetBlendSpaceGraph(Info.Payload);
+	UBlendSpaceGraph* BlendSpaceGraph = FTabPayload_BlendSpaceGraph::GetBlendSpaceGraph(Info.Payload);
 
-	return MakeAttributeLambda([WeakBlendSpace = TWeakObjectPtr<UBlendSpace>(DocumentID->BlendSpace)]()
+	return MakeAttributeLambda([WeakBlendSpace = TWeakObjectPtr<UBlendSpace>(BlendSpaceGraph->BlendSpace)]()
 		{ 
 			return WeakBlendSpace.Get() ? FText::FromName(WeakBlendSpace->GetFName()) : FText::GetEmpty();
 		});
@@ -273,24 +343,24 @@ TAttribute<FText> FBlendSpaceDocumentTabFactory::ConstructTabName(const FWorkflo
 
 void FBlendSpaceDocumentTabFactory::OnTabActivated(TSharedPtr<SDockTab> Tab) const
 {
-	TSharedRef<SBlendSpaceDocumentTab> DocumentWidget = StaticCastSharedRef<SBlendSpaceDocumentTab>(Tab->GetContent());
-	if(UBlendSpaceGraph* Document = DocumentWidget->Document.Get())
+	TSharedRef<SBlendSpaceDocumentTab> BlendSpaceDocumentTab = StaticCastSharedRef<SBlendSpaceDocumentTab>(Tab->GetContent());
+	if(UBlendSpaceGraph* BlendSpaceGraph = BlendSpaceDocumentTab->BlendSpaceGraph.Get())
 	{
 		if(TSharedPtr<FAnimationBlueprintEditor> BlueprintEditor = BlueprintEditorPtr.Pin())
 		{
-			BlueprintEditor->SetDetailObject(Document);
+			BlueprintEditor->SetDetailObject(BlendSpaceGraph);
 		}
 	}
 }
 
 void FBlendSpaceDocumentTabFactory::OnTabForegrounded(TSharedPtr<SDockTab> Tab) const
 {
-	TSharedRef<SBlendSpaceDocumentTab> DocumentWidget = StaticCastSharedRef<SBlendSpaceDocumentTab>(Tab->GetContent());
-	if(UBlendSpaceGraph* Document = DocumentWidget->Document.Get())
+	TSharedRef<SBlendSpaceDocumentTab> BlendSpaceDocumentTab = StaticCastSharedRef<SBlendSpaceDocumentTab>(Tab->GetContent());
+	if(UBlendSpaceGraph* BlendSpaceGraph = BlendSpaceDocumentTab->BlendSpaceGraph.Get())
 	{
 		if(TSharedPtr<FAnimationBlueprintEditor> BlueprintEditor = BlueprintEditorPtr.Pin())
 		{
-			BlueprintEditor->SetDetailObject(Document);
+			BlueprintEditor->SetDetailObject(BlendSpaceGraph);
 		}
 	}
 }
diff --git a/Engine/Source/Editor/BlueprintGraph/Private/BlueprintEditorSettings.cpp b/Engine/Source/Editor/BlueprintGraph/Private/BlueprintEditorSettings.cpp
--- a/Engine/Source/Editor/BlueprintGraph/Private/BlueprintEditorSettings.cpp
+++ b/Engine/Source/Editor/BlueprintGraph/Private/BlueprintEditorSettings.cpp
@@ -30,7 +30,7 @@ UBlueprintEditorSettings::UBlueprintEditorSettings(const FObjectInitializer& Obj
 	, bHideConstructionScriptComponentsInDetailsView(true)
 	, bHostFindInBlueprintsInGlobalTab(true)
 	, bNavigateToNativeFunctionsFromCallNodes(true)
-	, bDoubleClickNavigatesToParent(false)
+	, bDoubleClickNavigatesToParent(true)
 	, bEnableTypePromotion(true)
 	// Experimental
 	, bEnableNamespaceFilteringFeatures(false)
diff --git a/Engine/Source/Editor/Kismet/Private/BlueprintEditor.cpp b/Engine/Source/Editor/Kismet/Private/BlueprintEditor.cpp
--- a/Engine/Source/Editor/Kismet/Private/BlueprintEditor.cpp
+++ b/Engine/Source/Editor/Kismet/Private/BlueprintEditor.cpp
@@ -3323,51 +3323,88 @@ void FBlueprintEditor::NavigateToChildGraph()
 	{
 		UEdGraph* CurrentGraph = FocusedGraphEdPtr.Pin()->GetCurrentGraph();
 
-		if (CurrentGraph->SubGraphs.Num() > 1)
+		if (CurrentGraph->Nodes.Num() > 0)
 		{
 			// Display a child jump list
 			FMenuBuilder MenuBuilder(true, nullptr);
 			MenuBuilder.BeginSection("NavigateToGraph", LOCTEXT("ChildGraphPickerDesc", "Navigate to graph"));
 
-			TArray<TObjectPtr<UEdGraph>> SortedSubGraphs = CurrentGraph->SubGraphs;
-			SortedSubGraphs.Sort([](UEdGraph& A, UEdGraph& B) { return FLocalKismetCallbacks::GetGraphDisplayName(&A).CompareToCaseIgnored(FLocalKismetCallbacks::GetGraphDisplayName(&B)) < 0; });
+			TArray<TObjectPtr<UEdGraphNode>> SortedGraphNodes = CurrentGraph->Nodes;
+			SortedGraphNodes.Sort([](UEdGraphNode& A, UEdGraphNode& B) {
+				FText AName = A.GetNodeTitle(ENodeTitleType::ListView);
+				FText BName = B.GetNodeTitle(ENodeTitleType::ListView);
+				return AName.CompareToCaseIgnored(BName) < 0; });
 
-			for (TObjectPtr<UEdGraph> ChildGraph : SortedSubGraphs)
+			TObjectPtr<UEdGraphNode> SingleNode;
+			int32 NumEntries = 0;
+
+			for (TObjectPtr<UEdGraphNode> Node : SortedGraphNodes)
 			{
-				MenuBuilder.AddMenuEntry(
-					FLocalKismetCallbacks::GetGraphDisplayName(ChildGraph),
-					LOCTEXT("ChildGraphPickerTooltip", "Pick the graph to enter"),
-					FSlateIcon(),
-					FUIAction(
-						FExecuteAction::CreateLambda([this,ChildGraph]() {
-							 OpenDocument(ChildGraph, FDocumentTracker::NavigatingCurrentDocument);
-							 // focus blueprint widget so that keyboard inputs (navigate to child/parent) will work before clicking
-							 FSlateApplication::Get().SetKeyboardFocus(GetMyBlueprintWidget());
-							  }),
-						FCanExecuteAction()));
+				// Just calling CanJumpToDefinition isn't enough as it returns true for functions (resulting in a jump
+				// to code, which isn't desired).
+				UObject* TargetObject = Node->GetJumpTargetForDoubleClick();
+				if (TargetObject && Node->CanJumpToDefinition())
+				{
+					++NumEntries;
+					SingleNode = Node;
+					MenuBuilder.AddMenuEntry(
+						Node->GetNodeTitle(ENodeTitleType::ListView),
+						LOCTEXT("ChildGraphPickerTooltip", "Pick the graph to enter"),
+						FSlateIcon(),
+						FUIAction(
+							FExecuteAction::CreateLambda(
+								[this, Node]() 
+								{
+									if (Node->CanJumpToDefinition())
+									{
+										Node->JumpToDefinition();
+									}
+									// Note that setting the keyboard focus here doesn't work when navigating into a
+									// blendspace (i.e. a UEdGraph). Neither does trying to set the focus to 
+									// Node->DEPRECATED_NodeWidget
+									SetKeyboardFocus();
+								}),
+							FCanExecuteAction()));
+				}
 			}
 			MenuBuilder.EndSection();
-
-			FSlateApplication::Get().PushMenu( 
-				GetToolkitHost()->GetParentWidget(),
-				FWidgetPath(),
-				MenuBuilder.MakeWidget(),
-				FSlateApplication::Get().GetCursorPos(), // summon location
-				FPopupTransitionEffect( FPopupTransitionEffect::TypeInPopup )
-			);
+			
+			if (NumEntries > 0)
+			{
+				// If there is only one entry we could just jump straight to it. However, sometimes that can be a little
+				// disorientating if it was not obvious to the user exactly what the targets might be.
+				FSlateApplication::Get().PushMenu( 
+					GetToolkitHost()->GetParentWidget(),
+					FWidgetPath(),
+					MenuBuilder.MakeWidget(),
+					FSlateApplication::Get().GetCursorPos(), // summon location
+					FPopupTransitionEffect( FPopupTransitionEffect::TypeInPopup )
+				);
+			}
 		}
-		else if (CurrentGraph->SubGraphs.Num() == 1)
+	}
+}
+
+bool FBlueprintEditor::CanNavigateToChildGraph() const
+{
+	if (FocusedGraphEdPtr.IsValid())
+	{
+		UEdGraph* CurrentGraph = FocusedGraphEdPtr.Pin()->GetCurrentGraph();
+		for (TObjectPtr<UEdGraphNode> Node : CurrentGraph->Nodes)
 		{
-			// Jump immediately to the child if there is only one
-			UEdGraph* ChildGraph = CurrentGraph->SubGraphs[0];
-			OpenDocument(ChildGraph, FDocumentTracker::NavigatingCurrentDocument);
+			UObject* TargetObject = Node->GetJumpTargetForDoubleClick();
+			if (TargetObject && Node->CanJumpToDefinition())
+			{
+				return true;
+			}
 		}
 	}
+	return false;
 }
 
-bool FBlueprintEditor::CanNavigateToChildGraph() const
+void FBlueprintEditor::SetKeyboardFocus()
 {
-	return FocusedGraphEdPtr.IsValid() && (FocusedGraphEdPtr.Pin()->GetCurrentGraph()->SubGraphs.Num() > 0);
+	FSlateApplication::Get().SetKeyboardFocus(GetMyBlueprintWidget());
 }
 
 bool FBlueprintEditor::TransactionObjectAffectsBlueprint(UObject* InTransactedObject)
@@ -9882,12 +9919,18 @@ float FBlueprintEditor::GetInstructionTextOpacity(UEdGraph* InGraph) const
 	return 1.0f;
 }
 
+FText FBlueprintEditor::GetGraphDisplayName(const UEdGraph* Graph)
+{
+	return FLocalKismetCallbacks::GetGraphDisplayName(Graph);
+}
+
+
 FText FBlueprintEditor::GetGraphDecorationString(UEdGraph* InGraph) const
 {
 	return FText::GetEmpty();
 }
 
-bool FBlueprintEditor::IsGraphInCurrentBlueprint(UEdGraph* InGraph) const
+bool FBlueprintEditor::IsGraphInCurrentBlueprint(const UEdGraph* InGraph) const
 {
 	bool bEditable = true;
 
diff --git a/Engine/Source/Editor/Kismet/Private/BlueprintEditorTabFactories.h b/Engine/Source/Editor/Kismet/Private/BlueprintEditorTabFactories.h
--- a/Engine/Source/Editor/Kismet/Private/BlueprintEditorTabFactories.h
+++ b/Engine/Source/Editor/Kismet/Private/BlueprintEditorTabFactories.h
@@ -31,7 +31,7 @@ struct FLocalKismetCallbacks
 		return (Object != NULL) ? FText::FromString( Object->GetName() ) : LOCTEXT("UnknownObjectName", "UNKNOWN");
 	}
 
-	static FText GetGraphDisplayName(UEdGraph* Graph)
+	static FText GetGraphDisplayName(const UEdGraph* Graph)
 	{
 		if (Graph)
 		{
@@ -146,7 +146,7 @@ public:
 protected:
 	virtual TAttribute<FText> ConstructTabNameForObject(UEdGraph* DocumentID) const override
 	{
-		return TAttribute<FText>::Create(TAttribute<FText>::FGetter::CreateStatic<UEdGraph*>(&FLocalKismetCallbacks::GetGraphDisplayName, DocumentID));
+		return TAttribute<FText>::Create(TAttribute<FText>::FGetter::CreateStatic<const UEdGraph*>(&FLocalKismetCallbacks::GetGraphDisplayName, DocumentID));
 	}
 
 	virtual TSharedRef<SWidget> CreateTabBodyForObject(const FWorkflowTabSpawnInfo& Info, UEdGraph* DocumentID) const override;
diff --git a/Engine/Source/Editor/Kismet/Public/BlueprintEditor.h b/Engine/Source/Editor/Kismet/Public/BlueprintEditor.h
--- a/Engine/Source/Editor/Kismet/Public/BlueprintEditor.h
+++ b/Engine/Source/Editor/Kismet/Public/BlueprintEditor.h
@@ -271,6 +271,9 @@ public:
 	/**	Returns whether the editor is currently editing a single blueprint object */
 	bool IsEditingSingleBlueprint() const;
 
+	/** Ensures the blueprint has keyboard focus. */
+	void SetKeyboardFocus();
+
 	/** Getters for the various Kismet2 widgets */
 	TSharedRef<SKismetInspector> GetInspector() const { return Inspector.ToSharedRef(); }
 	TSharedRef<SKismetInspector> GetDefaultEditor() const { return DefaultEditor.ToSharedRef(); }
@@ -677,8 +680,11 @@ public:
 	/** Called by a graph title bar to get any extra information the editor would like to display */
 	virtual FText GetGraphDecorationString(UEdGraph* InGraph) const;
 
+	/** Gets the display name of a graph */
+	static FText GetGraphDisplayName(const UEdGraph* Graph);
+
 	/** Checks to see if the provided graph is contained within the current blueprint */
-	bool IsGraphInCurrentBlueprint(UEdGraph* InGraph) const;
+	bool IsGraphInCurrentBlueprint(const UEdGraph* InGraph) const;
 
 	/** Get the context to use from the Blueprint type */
 	static FName GetContextFromBlueprintType(EBlueprintType InType);
diff --git a/Engine/Source/Editor/Persona/Private/PersonaModule.cpp b/Engine/Source/Editor/Persona/Private/PersonaModule.cpp
--- a/Engine/Source/Editor/Persona/Private/PersonaModule.cpp
+++ b/Engine/Source/Editor/Persona/Private/PersonaModule.cpp
@@ -1498,6 +1498,9 @@ TSharedRef<SWidget> FPersonaModule::CreateBlendSpaceEditWidget(UBlendSpace* InBl
 	return SNew(SBlendSpaceEditor)
 		.BlendSpace(InBlendSpace)
 		.DisplayScrubBar(false)
+		.OnBlendSpaceNavigateUp(InArgs.OnBlendSpaceNavigateUp)
+		.OnBlendSpaceNavigateDown(InArgs.OnBlendSpaceNavigateDown)
+		.OnBlendSpaceCanvasDoubleClicked(InArgs.OnBlendSpaceCanvasDoubleClicked)
 		.OnBlendSpaceSampleDoubleClicked(InArgs.OnBlendSpaceSampleDoubleClicked)
 		.OnBlendSpaceSampleAdded(InArgs.OnBlendSpaceSampleAdded)
 		.OnBlendSpaceSampleRemoved(InArgs.OnBlendSpaceSampleRemoved)
diff --git a/Engine/Source/Editor/Persona/Private/SAnimationBlendSpace.cpp b/Engine/Source/Editor/Persona/Private/SAnimationBlendSpace.cpp
--- a/Engine/Source/Editor/Persona/Private/SAnimationBlendSpace.cpp
+++ b/Engine/Source/Editor/Persona/Private/SAnimationBlendSpace.cpp
@@ -85,6 +85,9 @@ void SBlendSpaceEditor::Construct(const FArguments& InArgs)
 							.OnSampleAdded(this, &SBlendSpaceEditor::OnSampleAdded)
 							.OnSampleDuplicated(this, &SBlendSpaceEditor::OnSampleDuplicated)
 							.OnSampleReplaced(this, &SBlendSpaceEditor::OnSampleReplaced)
+							.OnNavigateUp(InArgs._OnBlendSpaceNavigateUp)
+							.OnNavigateDown(InArgs._OnBlendSpaceNavigateDown)
+							.OnCanvasDoubleClicked(InArgs._OnBlendSpaceCanvasDoubleClicked)
 							.OnSampleDoubleClicked(InArgs._OnBlendSpaceSampleDoubleClicked)
 							.OnExtendSampleTooltip(InArgs._OnExtendSampleTooltip)
 							.OnGetBlendSpaceSampleName(InArgs._OnGetBlendSpaceSampleName)
@@ -279,7 +282,7 @@ void SBlendSpaceEditor::UpdatePreviewParameter() const
 			{
 				const FVector PreviewPosition = BlendSpaceGridWidget->GetPreviewPosition();
 				Component->PreviewInstance->SetBlendSpacePosition(PreviewPosition);
-				GetPreviewScene()->InvalidateViews();			
+				GetPreviewScene()->InvalidateViews();
 			}
 		}
 	}
diff --git a/Engine/Source/Editor/Persona/Private/SAnimationBlendSpace.h b/Engine/Source/Editor/Persona/Private/SAnimationBlendSpace.h
--- a/Engine/Source/Editor/Persona/Private/SAnimationBlendSpace.h
+++ b/Engine/Source/Editor/Persona/Private/SAnimationBlendSpace.h
@@ -26,6 +26,9 @@ public:
 		{}
 	SLATE_ARGUMENT(UBlendSpace*, BlendSpace)
 	SLATE_ARGUMENT(bool, DisplayScrubBar)
+	SLATE_EVENT(FOnBlendSpaceNavigateUp, OnBlendSpaceNavigateUp)
+	SLATE_EVENT(FOnBlendSpaceNavigateDown, OnBlendSpaceNavigateDown)
+	SLATE_EVENT(FOnBlendSpaceCanvasDoubleClicked, OnBlendSpaceCanvasDoubleClicked)
 	SLATE_EVENT(FOnBlendSpaceSampleDoubleClicked, OnBlendSpaceSampleDoubleClicked)
 	SLATE_EVENT(FOnBlendSpaceSampleAdded, OnBlendSpaceSampleAdded)
 	SLATE_EVENT(FOnBlendSpaceSampleDuplicated, OnBlendSpaceSampleDuplicated)
@@ -120,6 +123,6 @@ protected:
 	/** Flag to check whether or not the preview value should be (re-)set on the next tick */
 	bool bShouldSetPreviewPosition = true;
 
-	/** Cross-framed boundary scoped transaction, used for sticking to single transaction for interactive operations (dragging sample or spinbox)*/
+	/** Cross-framed boundary scoped transaction, used for sticking to single transaction for interactive operations (dragging sample or spinbox) */
 	TUniquePtr<FScopedTransaction> SampleMoveTransaction;
 };
diff --git a/Engine/Source/Editor/Persona/Private/SAnimationBlendSpaceGridWidget.cpp b/Engine/Source/Editor/Persona/Private/SAnimationBlendSpaceGridWidget.cpp
--- a/Engine/Source/Editor/Persona/Private/SAnimationBlendSpaceGridWidget.cpp
+++ b/Engine/Source/Editor/Persona/Private/SAnimationBlendSpaceGridWidget.cpp
@@ -196,6 +196,9 @@ void SBlendSpaceGridWidget::Construct(const FArguments& InArgs)
 	OnSampleMoved = InArgs._OnSampleMoved;
 	OnSampleRemoved = InArgs._OnSampleRemoved;
 	OnSampleReplaced = InArgs._OnSampleReplaced;
+	OnNavigateUp = InArgs._OnNavigateUp;
+	OnNavigateDown = InArgs._OnNavigateDown;
+	OnCanvasDoubleClicked = InArgs._OnCanvasDoubleClicked;
 	OnSampleDoubleClicked = InArgs._OnSampleDoubleClicked;
 	OnGetBlendSpaceSampleName = InArgs._OnGetBlendSpaceSampleName;
 	OnExtendSampleTooltip = InArgs._OnExtendSampleTooltip;
@@ -1204,7 +1207,10 @@ FReply SBlendSpaceGridWidget::OnMouseButtonDoubleClick(const FGeometry& InMyGeom
 			{
 				OnSampleDoubleClicked.ExecuteIfBound(SelectedSampleIndex);
 			}
-
+			else
+			{
+				OnCanvasDoubleClicked.ExecuteIfBound();
+			}
 			return FReply::Handled();
 		}
 	}
@@ -1362,26 +1368,34 @@ FReply SBlendSpaceGridWidget::OnKeyDown(const FGeometry& MyGeometry, const FKeyE
 {
 	if(!bReadOnly && BlendSpaceBase.IsSet())
 	{
-		if (BlendSpaceBase.IsSet())
+		// Start previewing when either one of the shift keys is pressed
+		if (IsHovered() && bMouseIsOverGeometry)
 		{
-			// Start previewing when either one of the shift keys is pressed
-			if (IsHovered() && bMouseIsOverGeometry)
+			if ((InKeyEvent.GetKey() == EKeys::LeftControl) || (InKeyEvent.GetKey() == EKeys::RightControl))
 			{
-				if ((InKeyEvent.GetKey() == EKeys::LeftControl) || (InKeyEvent.GetKey() == EKeys::RightControl))
-				{
-					StartPreviewing();
-					DragState = EDragState::Preview;
-					// Make tool tip visible (this will display the current preview sample value)
-					ShowToolTip();
-					return FReply::Handled();
-				}
+				StartPreviewing();
+				DragState = EDragState::Preview;
+				// Make tool tip visible (this will display the current preview sample value)
+				ShowToolTip();
+				return FReply::Handled();
+			}
 		
-				// Set flag for showing advanced preview info in tooltip
-				if ((InKeyEvent.GetKey() == EKeys::LeftAlt) || (InKeyEvent.GetKey() == EKeys::RightAlt))
-				{
-					bAdvancedPreview = true;
-					return FReply::Handled();
-				}
+			// Set flag for showing advanced preview info in tooltip
+			if ((InKeyEvent.GetKey() == EKeys::LeftAlt) || (InKeyEvent.GetKey() == EKeys::RightAlt))
+			{
+				bAdvancedPreview = true;
+				return FReply::Handled();
+			}
+			
+			if (InKeyEvent.GetKey() == EKeys::PageUp)
+			{
+				OnNavigateUp.ExecuteIfBound();
+				return FReply::Handled();
+			}
+			else if (InKeyEvent.GetKey() == EKeys::PageDown)
+			{
+				OnNavigateDown.ExecuteIfBound();
+				return FReply::Handled();
 			}
 		}
 	}
diff --git a/Engine/Source/Editor/Persona/Private/SAnimationBlendSpaceGridWidget.h b/Engine/Source/Editor/Persona/Private/SAnimationBlendSpaceGridWidget.h
--- a/Engine/Source/Editor/Persona/Private/SAnimationBlendSpaceGridWidget.h
+++ b/Engine/Source/Editor/Persona/Private/SAnimationBlendSpaceGridWidget.h
@@ -51,6 +51,9 @@ public:
 		SLATE_EVENT(FOnBlendSpaceSampleAdded, OnSampleAdded)
 		SLATE_EVENT(FOnBlendSpaceSampleDuplicated, OnSampleDuplicated)
 		SLATE_EVENT(FOnBlendSpaceSampleReplaced, OnSampleReplaced)
+		SLATE_EVENT(FOnBlendSpaceNavigateUp, OnNavigateUp)
+		SLATE_EVENT(FOnBlendSpaceNavigateDown, OnNavigateDown)
+		SLATE_EVENT(FOnBlendSpaceCanvasDoubleClicked, OnCanvasDoubleClicked)
 		SLATE_EVENT(FOnBlendSpaceSampleDoubleClicked, OnSampleDoubleClicked)
 		SLATE_EVENT(FOnGetBlendSpaceSampleName, OnGetBlendSpaceSampleName)
 		SLATE_ARGUMENT(FName, StatusBarName)
@@ -324,6 +327,9 @@ private:
 	FOnSampleMoved OnSampleMoved;
 	FOnBlendSpaceSampleRemoved OnSampleRemoved;
 	FOnBlendSpaceSampleReplaced OnSampleReplaced;
+	FOnBlendSpaceNavigateUp OnNavigateUp;
+	FOnBlendSpaceNavigateDown OnNavigateDown;
+	FOnBlendSpaceCanvasDoubleClicked OnCanvasDoubleClicked;
 	FOnBlendSpaceSampleDoubleClicked OnSampleDoubleClicked;
 	FOnGetBlendSpaceSampleName OnGetBlendSpaceSampleName;
 	FOnExtendBlendSpaceSampleTooltip OnExtendSampleTooltip;
diff --git a/Engine/Source/Editor/Persona/Public/PersonaDelegates.h b/Engine/Source/Editor/Persona/Public/PersonaDelegates.h
--- a/Engine/Source/Editor/Persona/Public/PersonaDelegates.h
+++ b/Engine/Source/Editor/Persona/Public/PersonaDelegates.h
@@ -37,6 +37,15 @@ DECLARE_DELEGATE_ThreeParams(FOnEditCurves, UAnimSequenceBase* /*InAnimSequence*
 // Called to stop editing curves in the curve editor
 DECLARE_DELEGATE_OneParam(FOnStopEditingCurves, const TArray<IAnimationEditor::FCurveEditInfo>& /*InCurveInfo*/);
 
+// Called when the user navigates with pageup
+DECLARE_DELEGATE(FOnBlendSpaceNavigateUp);
+
+// Called when the user navigates with pagedown
+DECLARE_DELEGATE(FOnBlendSpaceNavigateDown);
+
+// Called when the blendspace canvas is double clicked
+DECLARE_DELEGATE(FOnBlendSpaceCanvasDoubleClicked);
+
 // Called when a blendspace sample point is double clicked
 DECLARE_DELEGATE_OneParam(FOnBlendSpaceSampleDoubleClicked, int32 /*SampleIndex*/);
 
diff --git a/Engine/Source/Editor/Persona/Public/PersonaModule.h b/Engine/Source/Editor/Persona/Public/PersonaModule.h
--- a/Engine/Source/Editor/Persona/Public/PersonaModule.h
+++ b/Engine/Source/Editor/Persona/Public/PersonaModule.h
@@ -111,6 +111,15 @@ struct FBlendSpaceEditorArgs
 	// Called when a blendspace sample point is replaced
 	FOnBlendSpaceSampleReplaced OnBlendSpaceSampleReplaced;
 
+	// Called when the blendspace canvas is double clicked
+	FOnBlendSpaceNavigateUp OnBlendSpaceNavigateUp;
+
+	// Called when the blendspace canvas is double clicked
+	FOnBlendSpaceNavigateDown OnBlendSpaceNavigateDown;
+
+	// Called when the blendspace canvas is double clicked
+	FOnBlendSpaceCanvasDoubleClicked OnBlendSpaceCanvasDoubleClicked;
+
 	// Called when a blendspace sample point is double clicked
 	FOnBlendSpaceSampleDoubleClicked OnBlendSpaceSampleDoubleClicked;
 
