Function library support for member reference customizations and anim node functions

Also fixes missing error when using a non-thread-safe function call on an anim node
Also exposes CreateNodeFunctionsWidget on base anim graph node widget so that non-derived types can use it

diff --git a/Engine/Plugins/Runtime/PropertyAccess/Source/PropertyAccessEditor/Private/SPropertyBinding.cpp b/Engine/Plugins/Runtime/PropertyAccess/Source/PropertyAccessEditor/Private/SPropertyBinding.cpp
--- a/Engine/Plugins/Runtime/PropertyAccess/Source/PropertyAccessEditor/Private/SPropertyBinding.cpp
+++ b/Engine/Plugins/Runtime/PropertyAccess/Source/PropertyAccessEditor/Private/SPropertyBinding.cpp
@@ -25,6 +25,8 @@
 #include "Widgets/Input/SNumericEntryBox.h"
 #include "Algo/Transform.h"
 #include "Widgets/Layout/SSpacer.h"
+#include "UObject/UObjectIterator.h"
+#include "Kismet/BlueprintFunctionLibrary.h"
 
 #define LOCTEXT_NAMESPACE "PropertyBinding"
 
@@ -115,52 +117,70 @@ void SPropertyBinding::ForEachBindableFunction(UClass* FromClass, Predicate Pred
 {
 	if(Args.OnCanBindFunction.IsBound())
 	{
-		// Walk up class hierarchy for native functions and properties
-		for ( TFieldIterator<UFunction> FuncIt(FromClass, EFieldIteratorFlags::IncludeSuper); FuncIt; ++FuncIt )
+		auto CheckBindableClass = [this, &Pred](UClass* InBindableClass)
 		{
-			UFunction* Function = *FuncIt;
-
-			// Stop processing functions after reaching a base class that it doesn't make sense to go beyond.
-			if ( IsFieldFromBlackListedClass(Function) )
+			// Walk up class hierarchy for native functions and properties
+			for ( TFieldIterator<UFunction> FuncIt(InBindableClass, EFieldIteratorFlags::IncludeSuper); FuncIt; ++FuncIt )
 			{
-				break;
-			}
+				UFunction* Function = *FuncIt;
 
-			// Only allow binding pure functions if we're limited to pure function bindings.
-			if ( Args.bGeneratePureBindings && !Function->HasAnyFunctionFlags(FUNC_Const | FUNC_BlueprintPure) )
-			{
-				continue;
-			}
+				// Stop processing functions after reaching a base class that it doesn't make sense to go beyond.
+				if ( IsFieldFromBlackListedClass(Function) )
+				{
+					break;
+				}
 
-			// Only bind to functions that are callable from blueprints
-			if ( !UEdGraphSchema_K2::CanUserKismetCallFunction(Function) )
-			{
-				continue;
-			}
+				// Only allow binding pure functions if we're limited to pure function bindings.
+				if ( Args.bGeneratePureBindings && !Function->HasAnyFunctionFlags(FUNC_Const | FUNC_BlueprintPure) )
+				{
+					continue;
+				}
 
-			bool bValidObjectFunction = false;
-			if(Args.bAllowUObjectFunctions)
-			{
-				FObjectPropertyBase* ObjectPropertyBase = CastField<FObjectPropertyBase>(Function->GetReturnProperty());
-				if(ObjectPropertyBase != nullptr && Function->NumParms == 1)
+				// Only bind to functions that are callable from blueprints
+				if ( !UEdGraphSchema_K2::CanUserKismetCallFunction(Function) )
 				{
-					bValidObjectFunction = true;
+					continue;
 				}
-			}
 
-			bool bValidStructFunction = false;
-			if(Args.bAllowStructFunctions)
-			{
-				FStructProperty* StructProperty = CastField<FStructProperty>(Function->GetReturnProperty());
-				if(StructProperty != nullptr && Function->NumParms == 1)
+				bool bValidObjectFunction = false;
+				if(Args.bAllowUObjectFunctions)
+				{
+					FObjectPropertyBase* ObjectPropertyBase = CastField<FObjectPropertyBase>(Function->GetReturnProperty());
+					if(ObjectPropertyBase != nullptr && Function->NumParms == 1)
+					{
+						bValidObjectFunction = true;
+					}
+				}
+
+				bool bValidStructFunction = false;
+				if(Args.bAllowStructFunctions)
+				{
+					FStructProperty* StructProperty = CastField<FStructProperty>(Function->GetReturnProperty());
+					if(StructProperty != nullptr && Function->NumParms == 1)
+					{
+						bValidStructFunction = true;
+					}
+				}
+
+				if(bValidObjectFunction || bValidStructFunction || Args.OnCanBindFunction.Execute(Function))
 				{
-					bValidStructFunction = true;
+					Pred(FFunctionInfo(Function));
 				}
 			}
+		};
 
-			if(bValidObjectFunction || bValidStructFunction || Args.OnCanBindFunction.Execute(Function))
+		CheckBindableClass(FromClass);
+
+		if(Args.bAllowFunctionLibraryBindings)
+		{
+			// Iterate function libraries
+			for (TObjectIterator<UClass> ClassIt; ClassIt; ++ClassIt)
 			{
-				Pred(FFunctionInfo(Function));
+				UClass* TestClass = *ClassIt;
+				if (!TestClass->IsNative() && TestClass->IsChildOf(UBlueprintFunctionLibrary::StaticClass()) && (!TestClass->HasAnyClassFlags(CLASS_Abstract)))
+				{
+					CheckBindableClass(TestClass);
+				}
 			}
 		}
 	}
diff --git a/Engine/Source/Editor/AnimGraph/Private/AnimBlueprintExtension.cpp b/Engine/Source/Editor/AnimGraph/Private/AnimBlueprintExtension.cpp
--- a/Engine/Source/Editor/AnimGraph/Private/AnimBlueprintExtension.cpp
+++ b/Engine/Source/Editor/AnimGraph/Private/AnimBlueprintExtension.cpp
@@ -78,7 +78,8 @@ void UAnimBlueprintExtension::RequestExtensionsForNode(UAnimGraphNode_Base* InAn
             UAnimBlueprintExtension_PropertyAccess::StaticClass()
         };
 
-		if(InAnimGraphNode->InitialUpdateFunction.GetMemberGuid().IsValid() || InAnimGraphNode->BecomeRelevantFunction.GetMemberGuid().IsValid())
+		if(InAnimGraphNode->InitialUpdateFunction.ResolveMember<UFunction>(InAnimGraphNode->GetBlueprintClassFromNode()) != nullptr ||
+			InAnimGraphNode->BecomeRelevantFunction.ResolveMember<UFunction>(InAnimGraphNode->GetBlueprintClassFromNode()) != nullptr)
 		{
 			ExtensionClasses.Add(UAnimBlueprintExtension_NodeRelevancy::StaticClass());
 		}
diff --git a/Engine/Source/Editor/AnimGraph/Private/AnimGraphNode_Base.cpp b/Engine/Source/Editor/AnimGraph/Private/AnimGraphNode_Base.cpp
--- a/Engine/Source/Editor/AnimGraph/Private/AnimGraphNode_Base.cpp
+++ b/Engine/Source/Editor/AnimGraph/Private/AnimGraphNode_Base.cpp
@@ -203,21 +203,28 @@ void UAnimGraphNode_Base::ValidateAnimNodeDuringCompilation(USkeleton* ForSkelet
 		}
 	}
 
-	auto ValidateFunctionRef = [this, &MessageLog](const FMemberReference& InRef, const FText& InError)
+	auto ValidateFunctionRef = [this, &MessageLog](const FMemberReference& InRef, const FText& InFunctionName)
 	{
 		if(InRef.GetMemberName() != NAME_None)
 		{
 			UFunction* Function = InRef.ResolveMember<UFunction>(GetAnimBlueprint()->SkeletonGeneratedClass);
 			if(Function == nullptr)
 			{
-				MessageLog.Error(*InError.ToString(), this);
+				MessageLog.Error(*FText::Format(LOCTEXT("CouldNotResolveFunctionErrorFormat", "Could not resolve {0} function @@"), InFunctionName).ToString(), this);
+			}
+			else
+			{
+				if(!FBlueprintEditorUtils::HasFunctionBlueprintThreadSafeMetaData(Function))
+				{
+					MessageLog.Error(*FText::Format(LOCTEXT("FunctionThreadSafetyErrorFormat", "{0} function is not thread safe @@"), InFunctionName).ToString(), this);
+				}
 			}
 		}
 	};
 	
-	ValidateFunctionRef(InitialUpdateFunction, LOCTEXT("MissingInitialUpdateFunction", "Could not resolve initial update function for @@"));
-	ValidateFunctionRef(BecomeRelevantFunction, LOCTEXT("MissingBecomeRelevantFunction", "Could not resolve become relevant function for @@"));
-	ValidateFunctionRef(UpdateFunction, LOCTEXT("MissingUpdateFunction", "Could not resolve update function for @@"));
+	ValidateFunctionRef(InitialUpdateFunction, LOCTEXT("InitialUpdateFunctionName", "Initial Update"));
+	ValidateFunctionRef(BecomeRelevantFunction, LOCTEXT("BecomeRelevantFunctionName", "Become Relevant"));
+	ValidateFunctionRef(UpdateFunction, LOCTEXT("UpdateFunctionName", "Update"));
 }
 
 void UAnimGraphNode_Base::CopyTermDefaultsToDefaultObject(IAnimBlueprintCopyTermDefaultsContext& InCompilationContext, IAnimBlueprintNodeCopyTermDefaultsContext& InPerNodeContext, IAnimBlueprintGeneratedClassCompiledData& OutCompiledData)
@@ -501,9 +508,9 @@ void UAnimGraphNode_Base::ProcessDuringCompilation(IAnimBlueprintCompilationCont
 	Extension->ProcessNodePins(this, InCompilationContext, OutCompiledData);
 
 	// Resolve functions
-	GetFNode()->InitialUpdateFunction.SetFromFunction(InitialUpdateFunction.ResolveMember<UFunction>(GetAnimBlueprint()->SkeletonGeneratedClass));  
-	GetFNode()->BecomeRelevantFunction.SetFromFunction(BecomeRelevantFunction.ResolveMember<UFunction>(GetAnimBlueprint()->SkeletonGeneratedClass));
-	GetFNode()->UpdateFunction.SetFromFunction(UpdateFunction.ResolveMember<UFunction>(GetAnimBlueprint()->SkeletonGeneratedClass)); 
+	GetFNode()->InitialUpdateFunction.SetFromFunction(InitialUpdateFunction.ResolveMember<UFunction>(GetBlueprintClassFromNode()));
+	GetFNode()->BecomeRelevantFunction.SetFromFunction(BecomeRelevantFunction.ResolveMember<UFunction>(GetBlueprintClassFromNode()));
+	GetFNode()->UpdateFunction.SetFromFunction(UpdateFunction.ResolveMember<UFunction>(GetBlueprintClassFromNode())); 
 
 	// Call the override point
 	OnProcessDuringCompilation(InCompilationContext, OutCompiledData);
diff --git a/Engine/Source/Editor/AnimGraph/Public/AnimGraphNode_Base.h b/Engine/Source/Editor/AnimGraph/Public/AnimGraphNode_Base.h
--- a/Engine/Source/Editor/AnimGraph/Public/AnimGraphNode_Base.h
+++ b/Engine/Source/Editor/AnimGraph/Public/AnimGraphNode_Base.h
@@ -210,15 +210,15 @@ class ANIMGRAPH_API UAnimGraphNode_Base : public UK2Node
 	EBlueprintUsage BlueprintUsage;
 
 	// Function called before the node is updated for the first time
-	UPROPERTY(EditAnywhere, Category = Functions, meta=(FunctionReference, PrototypeFunction="/Script/AnimGraphRuntime.AnimExecutionContextLibrary.Prototype_ThreadSafeAnimUpdateCall"), DisplayName="On Initial Update")
+	UPROPERTY(EditAnywhere, Category = Functions, meta=(FunctionReference, AllowFunctionLibraries, PrototypeFunction="/Script/AnimGraphRuntime.AnimExecutionContextLibrary.Prototype_ThreadSafeAnimUpdateCall"), DisplayName="On Initial Update")
 	FMemberReference InitialUpdateFunction;
 
 	// Function called when the node becomes relevant
-	UPROPERTY(EditAnywhere, Category = Functions, meta=(FunctionReference, PrototypeFunction="/Script/AnimGraphRuntime.AnimExecutionContextLibrary.Prototype_ThreadSafeAnimUpdateCall"), DisplayName="On Become Relevant")
+	UPROPERTY(EditAnywhere, Category = Functions, meta=(FunctionReference, AllowFunctionLibraries, PrototypeFunction="/Script/AnimGraphRuntime.AnimExecutionContextLibrary.Prototype_ThreadSafeAnimUpdateCall"), DisplayName="On Become Relevant")
 	FMemberReference BecomeRelevantFunction;
 
 	// Function called when the node is updated
-	UPROPERTY(EditAnywhere, Category = Functions, meta=(FunctionReference, PrototypeFunction="/Script/AnimGraphRuntime.AnimExecutionContextLibrary.Prototype_ThreadSafeAnimUpdateCall"), DisplayName="On Update")
+	UPROPERTY(EditAnywhere, Category = Functions, meta=(FunctionReference, AllowFunctionLibraries, PrototypeFunction="/Script/AnimGraphRuntime.AnimExecutionContextLibrary.Prototype_ThreadSafeAnimUpdateCall"), DisplayName="On Update")
 	FMemberReference UpdateFunction;
 	
 	// UObject interface
diff --git a/Engine/Source/Editor/AnimationBlueprintEditor/Private/AnimationNodes/SAnimationGraphNode.cpp b/Engine/Source/Editor/AnimationBlueprintEditor/Private/AnimationNodes/SAnimationGraphNode.cpp
--- a/Engine/Source/Editor/AnimationBlueprintEditor/Private/AnimationNodes/SAnimationGraphNode.cpp
+++ b/Engine/Source/Editor/AnimationBlueprintEditor/Private/AnimationNodes/SAnimationGraphNode.cpp
@@ -280,37 +280,52 @@ void SAnimationGraphNode::CreateBelowPinControls(TSharedPtr<SVerticalBox> MainBo
 {
 	if (UAnimGraphNode_Base* AnimNode = CastChecked<UAnimGraphNode_Base>(GraphNode, ECastCheckedType::NullAllowed))
 	{
-		FPropertyEditorModule& PropertyEditorModule = FModuleManager::Get().LoadModuleChecked<FPropertyEditorModule>("PropertyEditor");
-		PropertyRowGenerator = PropertyEditorModule.CreatePropertyRowGenerator(FPropertyRowGeneratorArgs());
-		PropertyRowGenerator->RegisterInstancedCustomPropertyTypeLayout(FMemberReference::StaticStruct()->GetFName(), FOnGetPropertyTypeCustomizationInstance::CreateStatic(&FBlueprintMemberReferenceDetails::MakeInstance));
-		PropertyRowGenerator->SetObjects({ AnimNode });
-
-		TSharedPtr<SGridPanel> GridPanel;
+		TSharedRef<SWidget> NodeFunctionsWidget = CreateNodeFunctionsWidget(AnimNode, MakeAttributeSP(this, &SAnimationGraphNode::UseLowDetailNodeTitles));
 		
-		MainBox->AddSlot()
+		// Insert above the error reporting bar
+		MainBox->InsertSlot(FMath::Max(0, MainBox->NumSlots() - 1))
 		.AutoHeight()
 		.Padding(4.0f)
 		[
-			SNew(SLevelOfDetailBranchNode)
-			.UseLowDetailSlot(this, &SAnimationGraphNode::UseLowDetailNodeTitles)
-			.LowDetail()
+			SNew(SBox)
+			.IsEnabled_Lambda([this](){ return IsNodeEditable(); })
 			[
-				SNew(SSpacer)
-				.Size(FVector2D(17.0f, 17.f))
-			]
-			.HighDetail()
-			[
-				SAssignNew(GridPanel, SGridPanel)
-				.IsEnabled_Lambda([this](){ return IsNodeEditable(); })
+				NodeFunctionsWidget
 			]
 		];
+	}
+}
+
+// Widget used to allow functions to be viewed and edited on nodes
+class SAnimNodeFunctionsWidget : public SCompoundWidget
+{
+public:
+	SLATE_BEGIN_ARGS(SAnimNodeFunctionsWidget) {}
+
+	SLATE_ATTRIBUTE(bool, UseLowDetail)
+	
+	SLATE_END_ARGS()
+
+	void Construct(const FArguments& InArgs, UAnimGraphNode_Base* InNode)
+	{
+		FPropertyEditorModule& PropertyEditorModule = FModuleManager::Get().LoadModuleChecked<FPropertyEditorModule>("PropertyEditor");
+		PropertyRowGenerator = PropertyEditorModule.CreatePropertyRowGenerator(FPropertyRowGeneratorArgs());
+		PropertyRowGenerator->RegisterInstancedCustomPropertyTypeLayout(FMemberReference::StaticStruct()->GetFName(), FOnGetPropertyTypeCustomizationInstance::CreateStatic(&FBlueprintMemberReferenceDetails::MakeInstance));
+		PropertyRowGenerator->SetObjects({ InNode });
+
+		TSharedPtr<SGridPanel> GridPanel;
+
+		ChildSlot
+		[
+			SAssignNew(GridPanel, SGridPanel)
+		];
 
 		GridPanel->SetVisibility(EVisibility::Collapsed);
 
 		int32 RowIndex = 0;
 		
 		// Add bound functions
-		auto AddFunctionBindingWidget = [this, &GridPanel, &RowIndex](FName InCategory, FName InMemberName)
+		auto AddFunctionBindingWidget = [this, InNode, &GridPanel, &RowIndex](FName InCategory, FName InMemberName)
 		{
 			GridPanel->SetVisibility(EVisibility::Visible);
 			
@@ -332,9 +347,9 @@ void SAnimationGraphNode::CreateBelowPinControls(TSharedPtr<SVerticalBox> MainBo
 						{
 							DetailTreeNode = Children[ChildIdx];
 							PropertyHandle = ChildPropertyHandle;
-							PropertyHandle->SetOnPropertyValueChanged(FSimpleDelegate::CreateLambda([this]()
+							PropertyHandle->SetOnPropertyValueChanged(FSimpleDelegate::CreateLambda([InNode]()
 							{
-								GraphNode->ReconstructNode();
+								InNode->ReconstructNode();
 							}));
 							break;
 						}
@@ -353,7 +368,17 @@ void SAnimationGraphNode::CreateBelowPinControls(TSharedPtr<SVerticalBox> MainBo
 				.VAlign(VAlign_Center)
 				.Padding(10.0f, 2.0f, 2.0f, 2.0f)
 				[
-					NodeWidgets.NameWidget.ToSharedRef()
+					SNew(SLevelOfDetailBranchNode)
+					.UseLowDetailSlot(UseLowDetail)
+					.LowDetail()
+					[
+						SNew(SSpacer)
+						.Size(FVector2D(24.0f, 24.f))
+					]
+					.HighDetail()
+					[
+						NodeWidgets.NameWidget.ToSharedRef()
+					]
 				];
 
 				GridPanel->AddSlot(1, RowIndex)
@@ -361,26 +386,51 @@ void SAnimationGraphNode::CreateBelowPinControls(TSharedPtr<SVerticalBox> MainBo
 				.VAlign(VAlign_Center)
 				.Padding(2.0f, 2.0f, 10.0f, 2.0f)
 				[
-					NodeWidgets.ValueWidget.ToSharedRef()
+					SNew(SLevelOfDetailBranchNode)
+					.UseLowDetailSlot(UseLowDetail)
+					.LowDetail()
+					[
+						SNew(SSpacer)
+						.Size(FVector2D(24.0f, 24.f))
+					]
+					.HighDetail()
+					[	
+						NodeWidgets.ValueWidget.ToSharedRef()
+					]
 				];
 
 				RowIndex++;
 			}
 		};
 
-		if(AnimNode->InitialUpdateFunction.GetMemberGuid().IsValid())
+		if(InNode->InitialUpdateFunction.ResolveMember<UFunction>(InNode->GetBlueprintClassFromNode()) != nullptr)
 		{
 			AddFunctionBindingWidget("Functions", GET_MEMBER_NAME_CHECKED(UAnimGraphNode_Base, InitialUpdateFunction));
 		}
-		if(AnimNode->BecomeRelevantFunction.GetMemberGuid().IsValid())
+		if(InNode->BecomeRelevantFunction.ResolveMember<UFunction>(InNode->GetBlueprintClassFromNode()) != nullptr)
 		{
 			AddFunctionBindingWidget("Functions", GET_MEMBER_NAME_CHECKED(UAnimGraphNode_Base, BecomeRelevantFunction));
 		}
-		if(AnimNode->UpdateFunction.GetMemberGuid().IsValid())
+		if(InNode->UpdateFunction.ResolveMember<UFunction>(InNode->GetBlueprintClassFromNode()) != nullptr)
 		{
 			AddFunctionBindingWidget("Functions", GET_MEMBER_NAME_CHECKED(UAnimGraphNode_Base, UpdateFunction));
 		}
 	}
+
+	// Property row generator used to display function properties on nodes
+	TSharedPtr<IPropertyRowGenerator> PropertyRowGenerator;
+
+	// Hold a reference to the root ptr of the details tree we use to display function properties
+	TArray<TSharedPtr<IDetailTreeNode>> DetailNodes;
+
+	// Attribute allowing LOD
+	TAttribute<bool> UseLowDetail;
+};
+
+TSharedRef<SWidget> SAnimationGraphNode::CreateNodeFunctionsWidget(UAnimGraphNode_Base* InAnimNode, TAttribute<bool> InUseLowDetail)
+{
+	return SNew(SAnimNodeFunctionsWidget, InAnimNode)
+		.UseLowDetail(InUseLowDetail);
 }
 
 void SAnimationGraphNode::ReconfigurePinWidgetsForPropertyBindings(UAnimGraphNode_Base* InAnimGraphNode, TSharedRef<SGraphNode> InGraphNodeWidget, TFunctionRef<TSharedPtr<SGraphPin>(UEdGraphPin*)> InFindWidgetForPin)
diff --git a/Engine/Source/Editor/AnimationBlueprintEditor/Public/AnimationNodes/SAnimationGraphNode.h b/Engine/Source/Editor/AnimationBlueprintEditor/Public/AnimationNodes/SAnimationGraphNode.h
--- a/Engine/Source/Editor/AnimationBlueprintEditor/Public/AnimationNodes/SAnimationGraphNode.h
+++ b/Engine/Source/Editor/AnimationBlueprintEditor/Public/AnimationNodes/SAnimationGraphNode.h
@@ -25,6 +25,9 @@ public:
 	// Tweak any created pin widgets so they respond to bindings
 	static void ReconfigurePinWidgetsForPropertyBindings(UAnimGraphNode_Base* InAnimGraphNode, TSharedRef<SGraphNode> InGraphNodeWidget, TFunctionRef<TSharedPtr<SGraphPin>(UEdGraphPin*)> InFindWidgetForPin);
 
+	// Create below-widget controls for editing anim node functions
+	static TSharedRef<SWidget> CreateNodeFunctionsWidget(UAnimGraphNode_Base* InAnimNode, TAttribute<bool> InUseLowDetail);
+	
 protected:
 	// SWidget interface
 	virtual void Tick(const FGeometry& AllottedGeometry, const double InCurrentTime, const float InDeltaTime) override;
@@ -56,10 +59,4 @@ private:
 	TSharedPtr<SNodeTitle> NodeTitle;
 
 	TWeakObjectPtr<class UPoseWatch> PoseWatch;
-
-	// Property row generator used to display function properties on nodes
-	TSharedPtr<IPropertyRowGenerator> PropertyRowGenerator;
-	
-	// Hold a reference to the root ptr of the details tree we use to display function properties
-	TArray<TSharedPtr<IDetailTreeNode>> DetailNodes;
 };
diff --git a/Engine/Source/Editor/Kismet/Private/BlueprintMemberReferenceCustomization.cpp b/Engine/Source/Editor/Kismet/Private/BlueprintMemberReferenceCustomization.cpp
--- a/Engine/Source/Editor/Kismet/Private/BlueprintMemberReferenceCustomization.cpp
+++ b/Engine/Source/Editor/Kismet/Private/BlueprintMemberReferenceCustomization.cpp
@@ -42,6 +42,8 @@ void FBlueprintMemberReferenceDetails::CustomizeHeader(TSharedRef<IPropertyHandl
 			// Only function references are supported right now
 			return;
 		}
+
+		const bool bAllowFunctionLibraryReferences = InStructPropertyHandle->HasMetaData("AllowFunctionLibraries");
 		
 		// Try to find the prototype function
 		const FString& PrototypeFunctionName = InStructPropertyHandle->GetMetaData("PrototypeFunction");
@@ -116,7 +118,15 @@ void FBlueprintMemberReferenceDetails::CustomizeHeader(TSharedRef<IPropertyHandl
 				
 				check(StructData);
 				FMemberReference* MemberReference = static_cast<FMemberReference*>(StructData);
-				MemberReference->SetFromField<UFunction>(InBindingChain[0].Field.Get<UFunction>(), true);
+				UFunction* Function = InBindingChain[0].Field.Get<UFunction>();
+				UClass* OwnerClass = Function ? Function->GetOwnerClass() : nullptr;
+				bool bSelfContext = false;
+				if(OwnerClass != nullptr)
+				{
+					bSelfContext = (Blueprint->GeneratedClass != nullptr && Blueprint->GeneratedClass->IsChildOf(OwnerClass)) ||
+									(Blueprint->SkeletonGeneratedClass != nullptr && Blueprint->SkeletonGeneratedClass->IsChildOf(OwnerClass));
+				}
+				MemberReference->SetFromField<UFunction>(Function, bSelfContext);
 
 				InStructPropertyHandle->NotifyPostChange(EPropertyChangeType::ValueSet);
 			}
@@ -132,8 +142,7 @@ void FBlueprintMemberReferenceDetails::CustomizeHeader(TSharedRef<IPropertyHandl
 				
 				check(StructData);
 				FMemberReference* MemberReference = static_cast<FMemberReference*>(StructData);
-				FGuid EmptyGuid;
-				MemberReference->SetSelfMember(NAME_None, EmptyGuid);
+				*MemberReference = FMemberReference();
 
 				InStructPropertyHandle->NotifyPostChange(EPropertyChangeType::ValueSet);
 			}
@@ -232,6 +241,7 @@ void FBlueprintMemberReferenceDetails::CustomizeHeader(TSharedRef<IPropertyHandl
 		Args.CurrentBindingColor = FEditorStyle::GetSlateColor("Colors.Foreground").GetSpecifiedColor();
 		Args.bGeneratePureBindings = false;
 		Args.bAllowFunctionBindings = bFunctionReference;
+		Args.bAllowFunctionLibraryBindings = bAllowFunctionLibraryReferences;
 		Args.bAllowPropertyBindings = false;
 		Args.bAllowNewBindings = true;
 		Args.bAllowArrayElementBindings = false;
diff --git a/Engine/Source/Editor/UnrealEd/Public/IPropertyAccessEditor.h b/Engine/Source/Editor/UnrealEd/Public/IPropertyAccessEditor.h
--- a/Engine/Source/Editor/UnrealEd/Public/IPropertyAccessEditor.h
+++ b/Engine/Source/Editor/UnrealEd/Public/IPropertyAccessEditor.h
@@ -187,9 +187,12 @@ struct FPropertyBindingWidgetArgs
 	/** Whether to generate pure bindings */
 	bool bGeneratePureBindings = true;
 
-	/** Whether to allow function bindings */
+	/** Whether to allow function bindings (to "this" class of the blueprint in question) */
 	bool bAllowFunctionBindings = true;
 
+	/** Whether to allow function library bindings in addition to the passed-in blueprint's class */
+	bool bAllowFunctionLibraryBindings = false;
+	
 	/** Whether to allow property bindings */
 	bool bAllowPropertyBindings = true;	
 	
diff --git a/Engine/Source/Runtime/Engine/Classes/Animation/AnimNodeFunctionRef.h b/Engine/Source/Runtime/Engine/Classes/Animation/AnimNodeFunctionRef.h
--- a/Engine/Source/Runtime/Engine/Classes/Animation/AnimNodeFunctionRef.h
+++ b/Engine/Source/Runtime/Engine/Classes/Animation/AnimNodeFunctionRef.h
@@ -34,7 +34,7 @@ public:
 	void SetFromFunctionName(FName InName) { FunctionName = InName; }
 
 	// Set the function via a function
-	void SetFromFunction(UFunction* InFunction) { FunctionName = InFunction ? InFunction->GetFName() : NAME_None; }
+	void SetFromFunction(UFunction* InFunction);
 	
 	// Get the function name
 	FName GetFunctionName() const { return FunctionName; }
@@ -46,10 +46,18 @@ public:
 	bool IsValid() const { return Function != nullptr; }
 
 private:
+	// The name of the class to call the function with. If this is NAME_None, we assume this is a 'thiscall', if it is valid then we assume (and verify) we should call the function on a function library CDO.
+	UPROPERTY()
+	FName ClassName = NAME_None;
+
 	// The name of the function to call
 	UPROPERTY()
-	FName FunctionName = NAME_None;
+	FName FunctionName = NAME_None;	
 
+	// The class to use to call the function with, recovered by looking for a class of name FunctionName
+	UPROPERTY(Transient)
+	TObjectPtr<const UClass> Class = nullptr;
+	
 	// The function to call, recovered by looking for a function of name FunctionName
 	UPROPERTY(Transient)
 	TObjectPtr<UFunction> Function = nullptr;
diff --git a/Engine/Source/Runtime/Engine/Private/Animation/AnimNodeFunctionRef.cpp b/Engine/Source/Runtime/Engine/Private/Animation/AnimNodeFunctionRef.cpp
--- a/Engine/Source/Runtime/Engine/Private/Animation/AnimNodeFunctionRef.cpp
+++ b/Engine/Source/Runtime/Engine/Private/Animation/AnimNodeFunctionRef.cpp
@@ -6,12 +6,49 @@
 #include "Animation/AnimNodeReference.h"
 #include "Animation/AnimSubsystem_NodeRelevancy.h"
 #include "Animation/AnimInstance.h"
+#include "EngineLogs.h"
+
+void FAnimNodeFunctionRef::SetFromFunction(UFunction* InFunction)
+{
+	if(InFunction)
+	{
+		FunctionName = InFunction->GetFName();
+		if(UClass* OwnerClass = InFunction->GetOwnerClass())
+		{
+			ensureAlwaysMsgf(OwnerClass->IsChildOf(UAnimInstance::StaticClass()) || OwnerClass->IsChildOf(UBlueprintFunctionLibrary::StaticClass()), TEXT("Function class must derive from either UAnimInstance or UBlueprintFunctionLibrary"));
+			ClassName = *OwnerClass->GetPathName();
+		}
+		else
+		{
+			ClassName = NAME_None;
+		}
+	}
+	else
+	{
+		FunctionName = NAME_None;
+		ClassName = NAME_None;
+	}
+}
 
 void FAnimNodeFunctionRef::Initialize(const UClass* InClass)
 {
+	const UClass* ClassToUse = InClass;
+	
+	if(ClassName != NAME_None)
+	{
+		if(UClass* FoundClass = FindObject<UClass>(nullptr, *ClassName.ToString()))
+		{
+			Class = ClassToUse = FoundClass;
+		}
+		else
+		{
+			UE_LOG(LogAnimation, Warning, TEXT("Could not find function class %s"), *ClassName.ToString());
+		}
+	}
+	
 	if(FunctionName != NAME_None)
 	{
-		Function = InClass->FindFunctionByName(FunctionName);
+		Function = ClassToUse->FindFunctionByName(FunctionName);
 	}
 }
 
@@ -19,7 +56,14 @@ void FAnimNodeFunctionRef::Call(UObject* InObject, void* InParameters) const
 {
 	if(IsValid())
 	{
-		InObject->ProcessEvent(Function, InParameters);
+		if(Class)
+		{
+			Class->GetDefaultObject()->ProcessEvent(Function, InParameters);
+		}
+		else
+		{
+			InObject->ProcessEvent(Function, InParameters);
+		}
 	}
 }
 
