Improved appearance of below widget controls for anim nodes

Ordering of error reporting and functions switched
Incorporated debug vis for players & blendspaces into nodes
Fixed blendspace graphs not displaying functions
Also fixes regular BP composite nodes showing their error indicator in the title bar
Tweaked appearance of player nodes to incorporate their sliders

diff --git a/Engine/Source/Editor/AnimationBlueprintEditor/Private/AnimationNodes/SGraphNodeBlendSpaceGraph.cpp b/Engine/Source/Editor/AnimationBlueprintEditor/Private/AnimationNodes/SGraphNodeBlendSpaceGraph.cpp
--- a/Engine/Source/Editor/AnimationBlueprintEditor/Private/AnimationNodes/SGraphNodeBlendSpaceGraph.cpp
+++ b/Engine/Source/Editor/AnimationBlueprintEditor/Private/AnimationNodes/SGraphNodeBlendSpaceGraph.cpp
@@ -9,6 +9,8 @@
 #include "IDocumentation.h"
 #include "SBlendSpacePreview.h"
 #include "AnimationNodes/SAnimationGraphNode.h"
+#include "SLevelOfDetailBranchNode.h"
+#include "Widgets/Layout/SSpacer.h"
 
 #define LOCTEXT_NAMESPACE "SGraphNodeBlendSpaceGraph"
 
@@ -106,41 +108,48 @@ TSharedPtr<SToolTip> SGraphNodeBlendSpaceGraph::GetComplexTooltip()
 
 }
 
-void SGraphNodeBlendSpaceGraph::UpdateGraphNode()
+TSharedRef<SWidget> SGraphNodeBlendSpaceGraph::CreateNodeBody()
 {
-	SGraphNodeK2Composite::UpdateGraphNode();
-
-	// Extend to add below-widget controls (composite nodes do not do this by default)
-	SNodePanel::SNode::FNodeSlot* Slot = GetSlot(ENodeZone::Center);
-	check(Slot);
-	TSharedPtr<SWidget> CenterWidget = Slot->DetachWidget();
-
+	TSharedRef<SWidget> NodeBody = SGraphNodeK2Composite::CreateNodeBody();
+	
 	UAnimGraphNode_BlendSpaceGraphBase* BlendSpaceNode = CastChecked<UAnimGraphNode_BlendSpaceGraphBase>(GraphNode);
 
-	(*Slot)
-	[
-		SNew(SVerticalBox)
+	return SNew(SVerticalBox)
 		+SVerticalBox::Slot()
 		.AutoHeight()
 		[
-			CenterWidget.ToSharedRef()
+			NodeBody
 		]
 		+SVerticalBox::Slot()
 		.AutoHeight()
 		[
-			SNew(SBlendSpacePreview, CastChecked<UAnimGraphNode_Base>(GraphNode))
-			.OnGetBlendSpaceSampleName(FOnGetBlendSpaceSampleName::CreateLambda([this, WeakBlendSpaceNode = TWeakObjectPtr<UAnimGraphNode_BlendSpaceGraphBase>(BlendSpaceNode)](int32 InSampleIndex) -> FName
-			{
-				if(WeakBlendSpaceNode.Get())
-				{
-					UAnimGraphNode_BlendSpaceGraphBase* BlendSpaceNode = WeakBlendSpaceNode.Get();
-					return BlendSpaceNode->GetGraphs()[InSampleIndex]->GetFName();
-				}
-
-				return NAME_None;
-			}))
+			SAnimationGraphNode::CreateNodeFunctionsWidget(BlendSpaceNode, MakeAttributeSP(this, &SGraphNodeBlendSpaceGraph::UseLowDetailNodeTitles))
 		]
-	];
+		+SVerticalBox::Slot()
+		.AutoHeight()
+		[
+			SNew(SLevelOfDetailBranchNode)
+			.UseLowDetailSlot(this, &SGraphNodeBlendSpaceGraph::UseLowDetailNodeTitles)
+			.LowDetail()
+			[
+				SNew(SSpacer)
+				.Size(FVector2D(100.0f, 100.f))
+			]
+			.HighDetail()
+			[
+				SNew(SBlendSpacePreview, CastChecked<UAnimGraphNode_Base>(GraphNode))
+				.OnGetBlendSpaceSampleName(FOnGetBlendSpaceSampleName::CreateLambda([this, WeakBlendSpaceNode = TWeakObjectPtr<UAnimGraphNode_BlendSpaceGraphBase>(BlendSpaceNode)](int32 InSampleIndex) -> FName
+				{
+					if(WeakBlendSpaceNode.Get())
+					{
+						UAnimGraphNode_BlendSpaceGraphBase* BlendSpaceNode = WeakBlendSpaceNode.Get();
+						return BlendSpaceNode->GetGraphs()[InSampleIndex]->GetFName();
+					}
+
+					return NAME_None;
+				}))
+			]
+		];
 }
 
 #undef LOCTEXT_NAMESPACE
\ No newline at end of file
diff --git a/Engine/Source/Editor/AnimationBlueprintEditor/Private/AnimationNodes/SGraphNodeBlendSpaceGraph.h b/Engine/Source/Editor/AnimationBlueprintEditor/Private/AnimationNodes/SGraphNodeBlendSpaceGraph.h
--- a/Engine/Source/Editor/AnimationBlueprintEditor/Private/AnimationNodes/SGraphNodeBlendSpaceGraph.h
+++ b/Engine/Source/Editor/AnimationBlueprintEditor/Private/AnimationNodes/SGraphNodeBlendSpaceGraph.h
@@ -17,14 +17,11 @@ public:
 
 	void Construct(const FArguments& InArgs, class UAnimGraphNode_BlendSpaceGraphBase* InNode);
 
-	// SGraphNode interface
-	virtual void UpdateGraphNode() override;
-	// End of SGraphNode interface
-
 protected:
 	// SGraphNodeK2Composite interface
 	virtual UEdGraph* GetInnerGraph() const override;
-
+	virtual TSharedRef<SWidget> CreateNodeBody() override;
+	
 	// SGraphNode interface
 	TSharedPtr<SToolTip> GetComplexTooltip() override;
 };
diff --git a/Engine/Source/Editor/AnimationBlueprintEditor/Private/AnimationNodes/SGraphNodeBlendSpacePlayer.cpp b/Engine/Source/Editor/AnimationBlueprintEditor/Private/AnimationNodes/SGraphNodeBlendSpacePlayer.cpp
--- a/Engine/Source/Editor/AnimationBlueprintEditor/Private/AnimationNodes/SGraphNodeBlendSpacePlayer.cpp
+++ b/Engine/Source/Editor/AnimationBlueprintEditor/Private/AnimationNodes/SGraphNodeBlendSpacePlayer.cpp
@@ -5,6 +5,8 @@
 #include "AnimGraphNode_Base.h"
 #include "AnimGraphNode_BlendSpacePlayer.h"
 #include "Kismet2/BlueprintEditorUtils.h"
+#include "SLevelOfDetailBranchNode.h"
+#include "Widgets/Layout/SSpacer.h"
 
 void SGraphNodeBlendSpacePlayer::Construct(const FArguments& InArgs, UAnimGraphNode_Base* InNode)
 {
@@ -25,14 +27,27 @@ void SGraphNodeBlendSpacePlayer::Construct(const FArguments& InArgs, UAnimGraphN
 	}));
 }
 
-void SGraphNodeBlendSpacePlayer::CreateBelowWidgetControls(TSharedPtr<SVerticalBox> MainBox)
+void SGraphNodeBlendSpacePlayer::CreateBelowPinControls(TSharedPtr<SVerticalBox> MainBox)
 {
-	MainBox->AddSlot()
+	SAnimationGraphNode::CreateBelowPinControls(MainBox);
+
+	// Insert above the error reporting bar
+	MainBox->InsertSlot(FMath::Max(0, MainBox->NumSlots() - 1))
 	.AutoHeight()
 	.VAlign(VAlign_Fill)
 	.Padding(0.0f)
 	[
-		SNew(SBlendSpacePreview, CastChecked<UAnimGraphNode_Base>(GraphNode))
+		SNew(SLevelOfDetailBranchNode)
+		.UseLowDetailSlot(this, &SGraphNodeBlendSpacePlayer::UseLowDetailNodeTitles)
+		.LowDetail()
+		[
+			SNew(SSpacer)
+			.Size(FVector2D(100.0f, 100.f))
+		]
+		.HighDetail()
+		[
+			SNew(SBlendSpacePreview, CastChecked<UAnimGraphNode_Base>(GraphNode))
+		]
 	];
 }
 
diff --git a/Engine/Source/Editor/AnimationBlueprintEditor/Private/AnimationNodes/SGraphNodeBlendSpacePlayer.h b/Engine/Source/Editor/AnimationBlueprintEditor/Private/AnimationNodes/SGraphNodeBlendSpacePlayer.h
--- a/Engine/Source/Editor/AnimationBlueprintEditor/Private/AnimationNodes/SGraphNodeBlendSpacePlayer.h
+++ b/Engine/Source/Editor/AnimationBlueprintEditor/Private/AnimationNodes/SGraphNodeBlendSpacePlayer.h
@@ -21,7 +21,7 @@ public:
 	void Construct(const FArguments& InArgs, UAnimGraphNode_Base* InNode);
 
 	// SGraphNode interface
-	virtual void CreateBelowWidgetControls(TSharedPtr<SVerticalBox> MainBox) override;
+	virtual void CreateBelowPinControls(TSharedPtr<SVerticalBox> MainBox) override;
 	// End of SGraphNode interface
 
 protected:
diff --git a/Engine/Source/Editor/AnimationBlueprintEditor/Private/AnimationNodes/SGraphNodeSequencePlayer.cpp b/Engine/Source/Editor/AnimationBlueprintEditor/Private/AnimationNodes/SGraphNodeSequencePlayer.cpp
--- a/Engine/Source/Editor/AnimationBlueprintEditor/Private/AnimationNodes/SGraphNodeSequencePlayer.cpp
+++ b/Engine/Source/Editor/AnimationBlueprintEditor/Private/AnimationNodes/SGraphNodeSequencePlayer.cpp
@@ -8,6 +8,9 @@
 #include "Kismet2/BlueprintEditorUtils.h"
 #include "Animation/AnimNode_SequencePlayer.h"
 #include "AnimGraphNode_SequencePlayer.h"
+#include "Styling/StyleColors.h"
+#include "SLevelOfDetailBranchNode.h"
+#include "Widgets/Layout/SSpacer.h"
 
 /////////////////////////////////////////////////////
 // SGraphNodeSequencePlayer
@@ -71,22 +74,35 @@ void SGraphNodeSequencePlayer::UpdateGraphNode()
 	SGraphNode::UpdateGraphNode();
 }
 
-void SGraphNodeSequencePlayer::CreateBelowWidgetControls(TSharedPtr<SVerticalBox> MainBox)
+void SGraphNodeSequencePlayer::CreateBelowPinControls(TSharedPtr<SVerticalBox> MainBox)
 {
-	FLinearColor Yellow(0.9f, 0.9f, 0.125f);
-
-	MainBox->AddSlot()
+	SAnimationGraphNode::CreateBelowPinControls(MainBox);
+	
+	// Insert above the error reporting bar
+	MainBox->InsertSlot(FMath::Max(0, MainBox->NumSlots() - 1))
 	.AutoHeight()
 	.VAlign( VAlign_Fill )
-	.Padding(FMargin(0, 4, 0, 0))
+	.Padding(FMargin(4))
 	[
-		SNew(SSlider)
-		.ToolTipText(this, &SGraphNodeSequencePlayer::GetPositionTooltip)
-		.Visibility(this, &SGraphNodeSequencePlayer::GetSliderVisibility)
-		.Value(this, &SGraphNodeSequencePlayer::GetSequencePositionRatio)
-		.OnValueChanged(this, &SGraphNodeSequencePlayer::SetSequencePositionRatio)
-		.Locked(false)
-		.SliderHandleColor(Yellow)
+		SNew(SLevelOfDetailBranchNode)
+		.UseLowDetailSlot(this, &SGraphNodeSequencePlayer::UseLowDetailNodeTitles)
+		.LowDetail()
+		[
+			SNew(SSpacer)
+			.Size(FVector2D(16.0f, 16.f))
+		]
+		.HighDetail()
+		[
+			SNew(SSlider)
+			.Style(&FEditorStyle::Get().GetWidgetStyle<FSliderStyle>("AnimBlueprint.AssetPlayerSlider"))
+			.ToolTipText(this, &SGraphNodeSequencePlayer::GetPositionTooltip)
+			.Visibility(this, &SGraphNodeSequencePlayer::GetSliderVisibility)
+			.Value(this, &SGraphNodeSequencePlayer::GetSequencePositionRatio)
+			.OnValueChanged(this, &SGraphNodeSequencePlayer::SetSequencePositionRatio)
+			.Locked(false)
+			.SliderHandleColor(FStyleColors::White)
+			.SliderBarColor(FStyleColors::Foreground)
+		]
 	];
 }
 
diff --git a/Engine/Source/Editor/AnimationBlueprintEditor/Private/AnimationNodes/SGraphNodeSequencePlayer.h b/Engine/Source/Editor/AnimationBlueprintEditor/Private/AnimationNodes/SGraphNodeSequencePlayer.h
--- a/Engine/Source/Editor/AnimationBlueprintEditor/Private/AnimationNodes/SGraphNodeSequencePlayer.h
+++ b/Engine/Source/Editor/AnimationBlueprintEditor/Private/AnimationNodes/SGraphNodeSequencePlayer.h
@@ -25,7 +25,7 @@ public:
 
 	// SGraphNode interface
 	virtual void UpdateGraphNode() override;
-	virtual void CreateBelowWidgetControls(TSharedPtr<SVerticalBox> MainBox) override;
+	virtual void CreateBelowPinControls(TSharedPtr<SVerticalBox> MainBox) override;
 	// End of SGraphNode interface
 
 protected:
diff --git a/Engine/Source/Editor/EditorStyle/Private/StarshipStyle.cpp b/Engine/Source/Editor/EditorStyle/Private/StarshipStyle.cpp
--- a/Engine/Source/Editor/EditorStyle/Private/StarshipStyle.cpp
+++ b/Engine/Source/Editor/EditorStyle/Private/StarshipStyle.cpp
@@ -5296,6 +5296,15 @@ void FStarshipEditorStyle::FStyle::SetupPersonaStyle()
 		Set("BlendSpaceEditor.LabelBackground", new FSlateRoundedBoxBrush(FStyleColors::Background, FStyleColors::Foreground, 1.0f));
 		Set("BlendSpaceEditor.ZoomToFit", new IMAGE_BRUSH("GenericCurveEditor/Icons/FramingSelected_48x", Icon16x16));
 
+		// Asset player slider
+		FSliderStyle AssetPlayerSliderStyle = FSliderStyle()
+			.SetNormalBarImage(   FSlateRoundedBoxBrush(FStyleColors::InputOutline, 2.0f, FStyleColors::InputOutline, 1.0f))
+			.SetHoveredBarImage(  FSlateRoundedBoxBrush(FStyleColors::InputOutline, 2.0f, FStyleColors::InputOutline, 1.0f))
+			.SetNormalThumbImage(  FSlateRoundedBoxBrush(FStyleColors::Foreground, Icon8x8) )
+			.SetHoveredThumbImage( FSlateRoundedBoxBrush(FStyleColors::ForegroundHover, Icon8x8) )
+			.SetBarThickness(4.0f);
+		Set("AnimBlueprint.AssetPlayerSlider", AssetPlayerSliderStyle);
+		
 		const FButtonStyle BlueprintContextTargetsButtonStyle = FButtonStyle()
 			.SetNormal(IMAGE_BRUSH("Common/TreeArrow_Collapsed_Hovered", Icon10x10, FLinearColor(0.2f, 0.2f, 0.2f, 1.f)))
 			.SetHovered(IMAGE_BRUSH("Common/TreeArrow_Collapsed_Hovered", Icon10x10, FLinearColor(0.4f, 0.4f, 0.4f, 1.f)))
diff --git a/Engine/Source/Editor/GraphEditor/Public/KismetNodes/SGraphNodeK2Composite.h b/Engine/Source/Editor/GraphEditor/Public/KismetNodes/SGraphNodeK2Composite.h
--- a/Engine/Source/Editor/GraphEditor/Public/KismetNodes/SGraphNodeK2Composite.h
+++ b/Engine/Source/Editor/GraphEditor/Public/KismetNodes/SGraphNodeK2Composite.h
@@ -29,6 +29,5 @@ protected:
 	FText GetPreviewCornerText() const;
 	FText GetTooltipTextForNode() const;
 
-private:
-	TSharedRef<SWidget> CreateNodeBody();
+	virtual TSharedRef<SWidget> CreateNodeBody();
 };
diff --git a/Engine/Source/Editor/Persona/Private/AnimGraphNodeDetails.cpp b/Engine/Source/Editor/Persona/Private/AnimGraphNodeDetails.cpp
--- a/Engine/Source/Editor/Persona/Private/AnimGraphNodeDetails.cpp
+++ b/Engine/Source/Editor/Persona/Private/AnimGraphNodeDetails.cpp
@@ -96,6 +96,12 @@ void FAnimGraphNodeDetails::CustomizeDetails(class IDetailLayoutBuilder& DetailB
 		return;
 	}
 
+	// In the case of property row generators, we dont have a details view, so we can skip this work
+	if(DetailBuilder.GetDetailsView() == nullptr)
+	{
+		return;
+	}
+	
 	// customize anim graph node's own details if needed
 	AnimGraphNode->CustomizeDetails(DetailBuilder);
 
diff --git a/Engine/Source/Editor/Persona/Private/SAnimationBlendSpaceGridWidget.cpp b/Engine/Source/Editor/Persona/Private/SAnimationBlendSpaceGridWidget.cpp
--- a/Engine/Source/Editor/Persona/Private/SAnimationBlendSpaceGridWidget.cpp
+++ b/Engine/Source/Editor/Persona/Private/SAnimationBlendSpaceGridWidget.cpp
@@ -504,14 +504,15 @@ int32 SBlendSpaceGridWidget::OnPaint(const FPaintArgs& Args, const FGeometry& Al
 
 void SBlendSpaceGridWidget::PaintBackgroundAndGrid(const FGeometry& AllottedGeometry, const FSlateRect& MyCullingRect, FSlateWindowElementList& OutDrawElements, int32& DrawLayerId) const
 {
-	// Fill the background
-	FSlateDrawElement::MakeBox( OutDrawElements, DrawLayerId + 1, AllottedGeometry.ToPaintGeometry(), BackgroundImage );
-
 	if(const UBlendSpace* BlendSpace = BlendSpaceBase.Get())
 	{
 		// Create the grid
 		const FVector2D GridSize = CachedGridRectangle.GetSize();
 		const FVector2D GridOffset = CachedGridRectangle.GetTopLeft();
+
+		// Fill the background of the grid
+		FSlateDrawElement::MakeBox( OutDrawElements, DrawLayerId + 1, AllottedGeometry.ToPaintGeometry(GridOffset, GridSize), BackgroundImage );
+		
 		TArray<FVector2D> LinePoints;
 
 		// Draw grid lines
@@ -1741,7 +1742,7 @@ const FVector SBlendSpaceGridWidget::ScreenPositionToSampleValue(const FVector2D
 
 const FSlateRect SBlendSpaceGridWidget::GetGridRectangleFromGeometry(const FGeometry& MyGeometry)
 {
-	const float TopOffset = 20.0f; // Ideally we'd get the size of the buttons (showing the label/triangulation etc)
+	const float TopOffset = bReadOnly ? 0.0f : 20.0f; // Ideally we'd get the size of the buttons (showing the label/triangulation etc)
 	FSlateRect WindowRect = FSlateRect(0, TopOffset, MyGeometry.GetLocalSize().X, MyGeometry.GetLocalSize().Y);
 	if (!bStretchToFit)
 	{
