Fixed static analysis warning

diff --git a/Engine/Source/Editor/AnimationBlueprintEditor/Private/AnimationNodes/SAnimationGraphNode.cpp b/Engine/Source/Editor/AnimationBlueprintEditor/Private/AnimationNodes/SAnimationGraphNode.cpp
--- a/Engine/Source/Editor/AnimationBlueprintEditor/Private/AnimationNodes/SAnimationGraphNode.cpp
+++ b/Engine/Source/Editor/AnimationBlueprintEditor/Private/AnimationNodes/SAnimationGraphNode.cpp
@@ -451,7 +451,7 @@ TSharedRef<SWidget> SAnimationGraphNode::CreateNodeFunctionsWidget(UAnimGraphNod
 TSharedRef<SWidget> SAnimationGraphNode::CreateNodeTagWidget(UAnimGraphNode_Base* InAnimNode, TAttribute<bool> InUseLowDetail)
 {
 	return SNew(SLevelOfDetailBranchNode)
-		.Visibility_Lambda([InAnimNode](){ return InAnimNode->Tag != NAME_None ? EVisibility::Visible : EVisibility::Collapsed; })
+		.Visibility_Lambda([InAnimNode](){ return (InAnimNode->Tag != NAME_None) ? EVisibility::Visible : EVisibility::Collapsed; })
 		.UseLowDetailSlot(InUseLowDetail)
 		.LowDetail()
 		[
