Pose watch widget now shows on Blend Space and State Machine nodes

diff --git a/Engine/Source/Editor/AnimationBlueprintEditor/Private/AnimationNodes/SAnimationGraphNode.cpp b/Engine/Source/Editor/AnimationBlueprintEditor/Private/AnimationNodes/SAnimationGraphNode.cpp
--- a/Engine/Source/Editor/AnimationBlueprintEditor/Private/AnimationNodes/SAnimationGraphNode.cpp
+++ b/Engine/Source/Editor/AnimationBlueprintEditor/Private/AnimationNodes/SAnimationGraphNode.cpp
@@ -27,111 +27,10 @@
 #include "Widgets/Layout/SGridPanel.h"
 #include "Widgets/Text/SInlineEditableTextBlock.h"
 #include "SGraphPanel.h"
+#include "SPoseWatchOverlay.h"
 
 #define LOCTEXT_NAMESPACE "AnimationGraphNode"
 
-class SPoseViewColourPickerPopup : public SCompoundWidget
-{
-public:
-	SLATE_BEGIN_ARGS(SPoseViewColourPickerPopup)
-	{}
-	SLATE_ARGUMENT(TWeakObjectPtr< UPoseWatch >, PoseWatch)
-	SLATE_END_ARGS()
-
-	void Construct(const FArguments& InArgs)
-	{
-		PoseWatch = InArgs._PoseWatch;
-
-		static const FSlateColorBrush WhiteBrush(FLinearColor::White);
-
-		TArrayView<const FColor> PoseWatchColors = AnimationEditorUtils::GetPoseWatchColorPalette();
-		const int32 NumColors = PoseWatchColors.Num();
-
-		const int32 Rows = 2;
-		const int32 Columns = NumColors / Rows;
-
-		TSharedPtr<SVerticalBox> Layout = SNew(SVerticalBox);
-
-		for (int32 RowIndex = 0; RowIndex < Rows; ++RowIndex)
-		{
-			TSharedPtr<SHorizontalBox> Row = SNew(SHorizontalBox);
-
-			for (int32 RowItem = 0; RowItem < Columns; ++RowItem)
-			{
-				int32 ColorIndex = RowItem + (RowIndex * Columns);
-				FColor Color = PoseWatchColors[ColorIndex];
-				if (ColorIndex >= NumColors)
-				{
-					break;
-				}
-
-				Row->AddSlot()
-				.Padding(5.f, 2.f)
-				[
-					SNew(SButton)
-					.ButtonStyle(FAppStyle::Get(), "NoBorder")
-					.HAlign(HAlign_Center)
-					.OnClicked(this, &SPoseViewColourPickerPopup::NewPoseWatchColourPicked, Color)
-					[
-						SNew(SImage)
-						.Image(&WhiteBrush)
-						.DesiredSizeOverride(FVector2D(24, 24))
-						.ColorAndOpacity(Color)
-					]
-				];
-
-			}
-
-			Layout->AddSlot()
-			[
-				Row.ToSharedRef()
-			];
-		}
-
-		Layout->AddSlot()
-			.AutoHeight()
-			.Padding(5.f, 2.f)
-			[
-				SNew(SButton)
-				.Text(LOCTEXT("RemovePoseWatch", "Remove Pose Watch"))
-				.OnClicked(this, &SPoseViewColourPickerPopup::RemovePoseWatch)
-			];
-
-		this->ChildSlot
-			[
-				SNew(SBorder)
-				.BorderImage(FEditorStyle::GetBrush(TEXT("Menu.Background")))
-				.Padding(10)
-				[
-					Layout->AsShared()
-				]
-			];
-	}
-
-private:
-	FReply NewPoseWatchColourPicked(FColor NewColour)
-	{
-		if (UPoseWatch* CurPoseWatch = PoseWatch.Get())
-		{
-			AnimationEditorUtils::UpdatePoseWatchColour(CurPoseWatch, NewColour);
-		}
-		FSlateApplication::Get().DismissAllMenus();
-		return FReply::Handled();
-	}
-
-	FReply RemovePoseWatch()
-	{
-		if (UPoseWatch* CurPoseWatch = PoseWatch.Get())
-		{
-			AnimationEditorUtils::RemovePoseWatch(CurPoseWatch);
-		}
-		FSlateApplication::Get().DismissAllMenus();
-		return FReply::Handled();
-	}
-
-	TWeakObjectPtr<UPoseWatch> PoseWatch;
-};
-
 void SAnimationGraphNode::Construct(const FArguments& InArgs, UAnimGraphNode_Base* InNode)
 {
 	this->GraphNode = InNode;
@@ -150,30 +49,12 @@ void SAnimationGraphNode::Construct(const FArguments& InArgs, UAnimGraphNode_Bas
 		.ToolTip(IDocumentation::Get()->CreateToolTip(LOCTEXT("AnimGraphNodeIndicatorTooltip", "Fast path enabled: This node is not using any Blueprint calls to update its data."), NULL, TEXT("Shared/GraphNodes/Animation"), TEXT("GraphNode_FastPathInfo")))
 		.Visibility(EVisibility::Visible);
 
-	PoseViewWidget =
-		SNew(SButton)
-		.ToolTipText(LOCTEXT("TogglePoseWatchVisibility", "Click to toggle visibility"))
-		.OnClicked(this, &SAnimationGraphNode::TogglePoseWatchVisibility)
-		.ButtonColorAndOpacity(this, &SAnimationGraphNode::GetPoseViewColour)
-		[
-			SNew(SImage).Image(this, &SAnimationGraphNode::GetPoseViewIcon)
-		];
-
-	// Search for an enabled or disabled pose watch on this node
-	PoseWatch = AnimationEditorUtils::FindPoseWatchForNode(GraphNode);
 
-	// Register for pose watch changes
-	AnimationEditorUtils::OnPoseWatchesChanged().AddSP(this, &SAnimationGraphNode::HandlePoseWatchesChanged);
+	PoseViewWidget = SNew(SPoseWatchOverlay, InNode);
 
 	LastHighDetailSize = FVector2D::ZeroVector;
 }
 
-void SAnimationGraphNode::HandlePoseWatchesChanged(UAnimBlueprint* InAnimBlueprint, UEdGraphNode* InNode)
-{
-	// Search for an enabled or disabled pose watch on this node
-	PoseWatch = AnimationEditorUtils::FindPoseWatchForNode(GraphNode);
-}
-
 void SAnimationGraphNode::Tick(const FGeometry& AllottedGeometry, const double InCurrentTime, const float InDeltaTime)
 {
 	SGraphNodeK2Base::Tick(AllottedGeometry, InCurrentTime, InDeltaTime);
@@ -201,14 +82,11 @@ TArray<FOverlayWidgetInfo> SAnimationGraphNode::GetOverlayWidgets(bool bSelected
 			Widgets.Add(Info);
 		}
 
-		if (PoseWatch.IsValid())
+		if (PoseViewWidget->IsPoseWatchValid())
 		{
-			const FSlateBrush* ImageBrush = FEditorStyle::GetBrush("Level.VisibleIcon16x");
-
 			FOverlayWidgetInfo Info;
-			Info.OverlayOffset = FVector2D(0 - (ImageBrush->ImageSize.X * 0.5f), -(ImageBrush->ImageSize.Y * 0.5f));
+			Info.OverlayOffset = PoseViewWidget->GetOverlayOffset();
 			Info.Widget = PoseViewWidget;
-			
 			Widgets.Add(Info);
 		}
 	}
@@ -216,41 +94,6 @@ TArray<FOverlayWidgetInfo> SAnimationGraphNode::GetOverlayWidgets(bool bSelected
 	return Widgets;
 }
 
-const FSlateBrush* SAnimationGraphNode::GetPoseViewIcon() const
-{
-	return FEditorStyle::GetBrush(PoseWatch->GetIsVisible() ? "Level.VisibleIcon16x" : "Level.NotVisibleIcon16x");
-}
-
-FSlateColor SAnimationGraphNode::GetPoseViewColour() const
-{
-	UPoseWatch* CurPoseWatch = PoseWatch.Get();
-	if (CurPoseWatch)
-	{
-		FLinearColor OutColor = CurPoseWatch->GetColor();
-		OutColor.A = CurPoseWatch->GetShouldDeleteOnDeselect() ? 0.5 : 0.9;
-		return FSlateColor(OutColor);
-	}
-	return FSlateColor(FColor::White); //Need a return value but should never actually get here
-}
-
-FReply SAnimationGraphNode::TogglePoseWatchVisibility()
-{
-	PoseWatch->ToggleIsVisible();
-	return FReply::Handled();
-}
-
-FReply SAnimationGraphNode::SpawnColourPicker()
-{
-	FSlateApplication::Get().PushMenu(
-		SharedThis(this),
-		FWidgetPath(),
-		SNew(SPoseViewColourPickerPopup).PoseWatch(PoseWatch),
-		FSlateApplication::Get().GetCursorPos(),
-		FPopupTransitionEffect(FPopupTransitionEffect::TypeInPopup)
-		);
-
-	return FReply::Handled();
-}
 
 TSharedRef<SWidget> SAnimationGraphNode::CreateTitleWidget(TSharedPtr<SNodeTitle> InNodeTitle)
 {
diff --git a/Engine/Source/Editor/AnimationBlueprintEditor/Private/AnimationNodes/SGraphNodeBlendSpaceGraph.cpp b/Engine/Source/Editor/AnimationBlueprintEditor/Private/AnimationNodes/SGraphNodeBlendSpaceGraph.cpp
--- a/Engine/Source/Editor/AnimationBlueprintEditor/Private/AnimationNodes/SGraphNodeBlendSpaceGraph.cpp
+++ b/Engine/Source/Editor/AnimationBlueprintEditor/Private/AnimationNodes/SGraphNodeBlendSpaceGraph.cpp
@@ -11,6 +11,7 @@
 #include "AnimationNodes/SAnimationGraphNode.h"
 #include "SLevelOfDetailBranchNode.h"
 #include "Widgets/Layout/SSpacer.h"
+#include "SPoseWatchOverlay.h"
 
 #define LOCTEXT_NAMESPACE "SGraphNodeBlendSpaceGraph"
 
@@ -20,6 +21,8 @@ void SGraphNodeBlendSpaceGraph::Construct(const FArguments& InArgs, UAnimGraphNo
 
 	SetCursor(EMouseCursor::CardinalCross);
 
+	PoseWatchWidget = SNew(SPoseWatchOverlay, InNode);
+
 	UpdateGraphNode();
 
 	SAnimationGraphNode::ReconfigurePinWidgetsForPropertyBindings(CastChecked<UAnimGraphNode_Base>(GraphNode), SharedThis(this), [this](UEdGraphPin* InPin){ return FindWidgetForPin(InPin); });
@@ -32,6 +35,24 @@ UEdGraph* SGraphNodeBlendSpaceGraph::GetInnerGraph() const
 	return BlendSpaceNode->GetBlendSpaceGraph();
 }
 
+TArray<FOverlayWidgetInfo> SGraphNodeBlendSpaceGraph::GetOverlayWidgets(bool bSelected, const FVector2D& WidgetSize) const
+{
+	TArray<FOverlayWidgetInfo> Widgets;
+
+	if (UAnimGraphNode_Base* AnimNode = CastChecked<UAnimGraphNode_Base>(GraphNode, ECastCheckedType::NullAllowed))
+	{
+		if (PoseWatchWidget->IsPoseWatchValid())
+		{
+			FOverlayWidgetInfo Info;
+			Info.OverlayOffset = PoseWatchWidget->GetOverlayOffset();
+			Info.Widget = PoseWatchWidget;
+			Widgets.Add(Info);
+		}
+	}
+
+	return Widgets;
+}
+
 TSharedPtr<SToolTip> SGraphNodeBlendSpaceGraph::GetComplexTooltip()
 {
 	if (UBlendSpaceGraph* BlendSpaceGraph = CastChecked<UBlendSpaceGraph>(GetInnerGraph()))
diff --git a/Engine/Source/Editor/AnimationBlueprintEditor/Private/AnimationNodes/SGraphNodeBlendSpaceGraph.h b/Engine/Source/Editor/AnimationBlueprintEditor/Private/AnimationNodes/SGraphNodeBlendSpaceGraph.h
--- a/Engine/Source/Editor/AnimationBlueprintEditor/Private/AnimationNodes/SGraphNodeBlendSpaceGraph.h
+++ b/Engine/Source/Editor/AnimationBlueprintEditor/Private/AnimationNodes/SGraphNodeBlendSpaceGraph.h
@@ -8,6 +8,7 @@
 
 class UEdGraph;
 class UAnimGraphNode_BlendSpaceGraphBase;
+class SPoseWatchOverlay;
 
 class SGraphNodeBlendSpaceGraph : public SGraphNodeK2Composite
 {
@@ -20,8 +21,12 @@ public:
 protected:
 	// SGraphNodeK2Composite interface
 	virtual UEdGraph* GetInnerGraph() const override;
+	virtual TArray<FOverlayWidgetInfo> GetOverlayWidgets(bool bSelected, const FVector2D& WidgetSize) const override;
 	virtual TSharedRef<SWidget> CreateNodeBody() override;
 	
 	// SGraphNode interface
 	TSharedPtr<SToolTip> GetComplexTooltip() override;
+
+private:
+	TSharedPtr<SPoseWatchOverlay> PoseWatchWidget;
 };
diff --git a/Engine/Source/Editor/AnimationBlueprintEditor/Private/AnimationNodes/SGraphNodeStateMachineInstance.cpp b/Engine/Source/Editor/AnimationBlueprintEditor/Private/AnimationNodes/SGraphNodeStateMachineInstance.cpp
--- a/Engine/Source/Editor/AnimationBlueprintEditor/Private/AnimationNodes/SGraphNodeStateMachineInstance.cpp
+++ b/Engine/Source/Editor/AnimationBlueprintEditor/Private/AnimationNodes/SGraphNodeStateMachineInstance.cpp
@@ -3,6 +3,13 @@
 #include "AnimationNodes/SGraphNodeStateMachineInstance.h"
 #include "AnimationStateMachineGraph.h"
 #include "AnimGraphNode_StateMachineBase.h"
+#include "Engine/PoseWatch.h"
+#include "AnimationEditorUtils.h"
+#include "Widgets/Input/SButton.h"
+#include "Widgets/Images/SImage.h"
+#include "SPoseWatchOverlay.h"
+
+#define LOCTEXT_NAMESPACE "SGraphNodeStateMachineInstance"
 
 /////////////////////////////////////////////////////
 // SGraphNodeStateMachineInstance
@@ -13,6 +20,8 @@ void SGraphNodeStateMachineInstance::Construct(const FArguments& InArgs, UAnimGr
 
 	SetCursor(EMouseCursor::CardinalCross);
 
+	PoseWatchWidget = SNew(SPoseWatchOverlay, InNode);
+
 	UpdateGraphNode();
 }
 
@@ -22,3 +31,23 @@ UEdGraph* SGraphNodeStateMachineInstance::GetInnerGraph() const
 
 	return StateMachineInstance->EditorStateMachineGraph;
 }
+
+TArray<FOverlayWidgetInfo> SGraphNodeStateMachineInstance::GetOverlayWidgets(bool bSelected, const FVector2D& WidgetSize) const
+{
+	TArray<FOverlayWidgetInfo> Widgets;
+
+	if (UAnimGraphNode_Base* AnimNode = CastChecked<UAnimGraphNode_Base>(GraphNode, ECastCheckedType::NullAllowed))
+	{
+		if (PoseWatchWidget->IsPoseWatchValid())
+		{
+			FOverlayWidgetInfo Info;
+			Info.OverlayOffset = PoseWatchWidget->GetOverlayOffset();
+			Info.Widget = PoseWatchWidget;
+			Widgets.Add(Info);
+		}
+	}
+
+	return Widgets;
+}
+
+#undef LOCTEXT_NAMESPACE
\ No newline at end of file
diff --git a/Engine/Source/Editor/AnimationBlueprintEditor/Private/AnimationNodes/SGraphNodeStateMachineInstance.h b/Engine/Source/Editor/AnimationBlueprintEditor/Private/AnimationNodes/SGraphNodeStateMachineInstance.h
--- a/Engine/Source/Editor/AnimationBlueprintEditor/Private/AnimationNodes/SGraphNodeStateMachineInstance.h
+++ b/Engine/Source/Editor/AnimationBlueprintEditor/Private/AnimationNodes/SGraphNodeStateMachineInstance.h
@@ -7,6 +7,7 @@
 #include "KismetNodes/SGraphNodeK2Composite.h"
 
 class UEdGraph;
+class SPoseWatchOverlay;
 
 class SGraphNodeStateMachineInstance : public SGraphNodeK2Composite
 {
@@ -15,8 +16,13 @@ public:
 	SLATE_END_ARGS()
 
 	void Construct(const FArguments& InArgs, class UAnimGraphNode_StateMachineBase* InNode);
+
 protected:
 	// SGraphNodeK2Composite interface
 	virtual UEdGraph* GetInnerGraph() const override;
+	virtual TArray<FOverlayWidgetInfo> GetOverlayWidgets(bool bSelected, const FVector2D& WidgetSize) const override;
 	// End of SGraphNodeK2Composite interface
+
+private:
+	TSharedPtr<SPoseWatchOverlay> PoseWatchWidget;
 };
diff --git a/Engine/Source/Editor/AnimationBlueprintEditor/Private/AnimationNodes/SPoseWatchOverlay.cpp b/Engine/Source/Editor/AnimationBlueprintEditor/Private/AnimationNodes/SPoseWatchOverlay.cpp
new file mode 100644
--- /dev/null
+++ b/Engine/Source/Editor/AnimationBlueprintEditor/Private/AnimationNodes/SPoseWatchOverlay.cpp
@@ -0,0 +1,81 @@
+// Copyright Epic Games, Inc. All Rights Reserved.
+
+#include "SPoseWatchOverlay.h"
+#include "Widgets/SBoxPanel.h"
+#include "Engine/PoseWatch.h"
+#include "Widgets/Images/SImage.h"
+#include "AnimationEditorUtils.h"
+#include "SNodePanel.h"
+
+#define LOCTEXT_NAMESPACE "SPoseWatchOverlay"
+
+const FSlateBrush* SPoseWatchOverlay::IconVisible = FEditorStyle::GetBrush("Level.VisibleIcon16x");
+const FSlateBrush* SPoseWatchOverlay::IconNotVisible = FEditorStyle::GetBrush("Level.NotVisibleIcon16x");
+
+void SPoseWatchOverlay::Construct(const FArguments& InArgs, UEdGraphNode* InNode)
+{
+	GraphNode = InNode;
+
+	PoseWatch = AnimationEditorUtils::FindPoseWatchForNode(InNode);
+	AnimationEditorUtils::OnPoseWatchesChanged().AddSP(this, &SPoseWatchOverlay::HandlePoseWatchesChanged);
+
+	ChildSlot
+	[
+		SNew(SButton)
+		.ToolTipText(LOCTEXT("TogglePoseWatchVisibility", "Click to toggle visibility"))
+		.OnClicked(this, &SPoseWatchOverlay::TogglePoseWatchVisibility)
+		.ButtonColorAndOpacity(this, &SPoseWatchOverlay::GetPoseViewColor)
+		[
+			SNew(SImage).Image(this, &SPoseWatchOverlay::GetPoseViewIcon)
+		]
+	];
+}
+
+void SPoseWatchOverlay::HandlePoseWatchesChanged(UAnimBlueprint* InAnimBlueprint, UEdGraphNode* InNode)
+{
+	PoseWatch = AnimationEditorUtils::FindPoseWatchForNode(GraphNode.Get());
+}
+
+FSlateColor SPoseWatchOverlay::GetPoseViewColor() const
+{
+	static constexpr float AlphaTemporary = 0.5f;
+	static constexpr float AlphaPermanent = 0.9f;
+
+	UPoseWatch* CurPoseWatch = PoseWatch.Get();
+	if (CurPoseWatch)
+	{
+		FLinearColor OutColor = CurPoseWatch->GetColor();
+		OutColor.A = CurPoseWatch->GetShouldDeleteOnDeselect() ? AlphaTemporary : AlphaPermanent;
+		return FSlateColor(OutColor);
+	}
+	return FSlateColor(FColor::Black);
+}
+
+const FSlateBrush* SPoseWatchOverlay::GetPoseViewIcon() const
+{
+	return (PoseWatch.IsValid() && PoseWatch->GetIsVisible()) ? IconVisible : IconNotVisible;
+}
+
+FReply SPoseWatchOverlay::TogglePoseWatchVisibility()
+{
+	if (PoseWatch.IsValid())
+	{
+		PoseWatch->ToggleIsVisible();
+		return FReply::Handled();
+	}
+	return FReply::Unhandled();
+}
+
+FVector2D SPoseWatchOverlay::GetOverlayOffset() const
+{
+	return FVector2D(0 - (IconVisible->ImageSize.X * 0.5f), -(IconVisible->ImageSize.Y * 0.5f));
+}
+
+bool SPoseWatchOverlay::IsPoseWatchValid() const
+{
+	return PoseWatch.IsValid();
+}
+
+
+
+#undef LOCTEXT_NAMESPACE
\ No newline at end of file
diff --git a/Engine/Source/Editor/AnimationBlueprintEditor/Private/AnimationNodes/SPoseWatchOverlay.h b/Engine/Source/Editor/AnimationBlueprintEditor/Private/AnimationNodes/SPoseWatchOverlay.h
new file mode 100644
--- /dev/null
+++ b/Engine/Source/Editor/AnimationBlueprintEditor/Private/AnimationNodes/SPoseWatchOverlay.h
@@ -0,0 +1,36 @@
+// Copyright Epic Games, Inc. All Rights Reserved.
+
+#pragma once
+
+#include "CoreMinimal.h"
+#include "Widgets/SWidget.h"
+#include "Widgets/Input/SButton.h"
+
+class UAnimBlueprint;
+class UEdGraphNode;
+struct FOverlayWidgetInfo;
+class UAnimGraphNode_Base;
+
+class SPoseWatchOverlay : public SCompoundWidget
+{
+public:
+	SLATE_BEGIN_ARGS(SPoseWatchOverlay) {}
+	SLATE_END_ARGS()
+
+	void Construct(const FArguments& InArgs, UEdGraphNode* InNode);
+	FVector2D GetOverlayOffset() const;
+	bool IsPoseWatchValid() const;
+
+private:
+	void HandlePoseWatchesChanged(UAnimBlueprint* InAnimBlueprint, UEdGraphNode* InNode);
+
+	FSlateColor GetPoseViewColor() const;
+	const FSlateBrush* GetPoseViewIcon() const;
+	FReply TogglePoseWatchVisibility();
+
+	TWeakObjectPtr<UEdGraphNode> GraphNode;
+	TWeakObjectPtr<class UPoseWatch> PoseWatch;
+
+	static const FSlateBrush* IconVisible;
+	static const FSlateBrush* IconNotVisible;
+};
\ No newline at end of file
diff --git a/Engine/Source/Editor/AnimationBlueprintEditor/Public/AnimationNodes/SAnimationGraphNode.h b/Engine/Source/Editor/AnimationBlueprintEditor/Public/AnimationNodes/SAnimationGraphNode.h
--- a/Engine/Source/Editor/AnimationBlueprintEditor/Public/AnimationNodes/SAnimationGraphNode.h
+++ b/Engine/Source/Editor/AnimationBlueprintEditor/Public/AnimationNodes/SAnimationGraphNode.h
@@ -14,6 +14,7 @@ class UAnimGraphNode_Base;
 class IDetailTreeNode;
 class IPropertyRowGenerator;
 class UAnimBlueprint;
+class SPoseWatchOverlay;
 
 class ANIMATIONBLUEPRINTEDITOR_API SAnimationGraphNode : public SGraphNodeK2Base
 {
@@ -53,13 +54,6 @@ protected:
 	// End of SGraphNode interface
 
 private:
-	// Return Pose View Colour for slate indicator
-	FSlateColor GetPoseViewColour() const;
-	const FSlateBrush* GetPoseViewIcon() const;
-
-	FReply SpawnColourPicker();
-	FReply TogglePoseWatchVisibility();
-
 	// Handle the node informing us that the title has changed
 	void HandleNodeTitleChanged();
 
@@ -74,13 +68,11 @@ private:
 	TSharedPtr<SWidget> IndicatorWidget;
 
 	/** Keep a reference to the pose view indicator widget handing around */
-	TSharedPtr<SWidget> PoseViewWidget;
+	TSharedPtr<SPoseWatchOverlay> PoseViewWidget;
 
 	/** Cache the node title so we can invalidate it */
 	TSharedPtr<SNodeTitle> NodeTitle;
 
-	TWeakObjectPtr<class UPoseWatch> PoseWatch;
-
 	/** Cached size from when we last drew at high detail */
 	FVector2D LastHighDetailSize;
 
