Fix crash when recompiling an animation blueprint with nodes that have 3d editing selected.

diff --git a/Engine/Source/Editor/AnimationBlueprintEditor/Private/AnimationBlueprintEditor.cpp b/Engine/Source/Editor/AnimationBlueprintEditor/Private/AnimationBlueprintEditor.cpp
--- a/Engine/Source/Editor/AnimationBlueprintEditor/Private/AnimationBlueprintEditor.cpp
+++ b/Engine/Source/Editor/AnimationBlueprintEditor/Private/AnimationBlueprintEditor.cpp
@@ -1821,10 +1821,9 @@ void FAnimationBlueprintEditor::OnSelectedNodesChangedImpl(const TSet<class UObj
 		for (FSelectedNodePtr CurrentAnimGraphNode : SelectedAnimGraphNodes)
 		{
 			UAnimGraphNode_Base* const CurrentAnimGraphNodePtr = CurrentAnimGraphNode.Get();
-			if (FAnimNode_Base* const PreviewNode = FindAnimNode(CurrentAnimGraphNodePtr))
-			{
-				CurrentAnimGraphNodePtr->OnNodeSelected(false, *PersonaEditorModeManager, PreviewNode);
-			}
+			FAnimNode_Base* const PreviewNode = FindAnimNode(CurrentAnimGraphNodePtr);
+			// intentionally not null checking PreviewNode here, in order to deselect nodes that are no longer included in the runtime graph after recompile
+			CurrentAnimGraphNodePtr->OnNodeSelected(false, *PersonaEditorModeManager, PreviewNode);
 		}
 
 		// Remove all the nodes that are no longer selected.
