- Fixed SPoseWatchOverlay to call GetBrush in Construct instead of global ctor time (for monolithic editor)

diff --git a/Engine/Source/Editor/AnimationBlueprintEditor/Private/AnimationNodes/SPoseWatchOverlay.cpp b/Engine/Source/Editor/AnimationBlueprintEditor/Private/AnimationNodes/SPoseWatchOverlay.cpp
--- a/Engine/Source/Editor/AnimationBlueprintEditor/Private/AnimationNodes/SPoseWatchOverlay.cpp
+++ b/Engine/Source/Editor/AnimationBlueprintEditor/Private/AnimationNodes/SPoseWatchOverlay.cpp
@@ -9,11 +9,14 @@
 
 #define LOCTEXT_NAMESPACE "SPoseWatchOverlay"
 
-const FSlateBrush* SPoseWatchOverlay::IconVisible = FEditorStyle::GetBrush("Level.VisibleIcon16x");
-const FSlateBrush* SPoseWatchOverlay::IconNotVisible = FEditorStyle::GetBrush("Level.NotVisibleIcon16x");
+const FSlateBrush* SPoseWatchOverlay::IconVisible = nullptr;
+const FSlateBrush* SPoseWatchOverlay::IconNotVisible = nullptr;
 
 void SPoseWatchOverlay::Construct(const FArguments& InArgs, UEdGraphNode* InNode)
 {
+	IconVisible = FEditorStyle::GetBrush("Level.VisibleIcon16x");
+	IconNotVisible = FEditorStyle::GetBrush("Level.NotVisibleIcon16x");
+
 	GraphNode = InNode;
 
 	PoseWatch = AnimationEditorUtils::FindPoseWatchForNode(InNode);
