Support for disabling pose watching on specific anim nodes

diff --git a/Engine/Source/Editor/AnimGraph/Private/AnimGraphNode_Root.cpp b/Engine/Source/Editor/AnimGraph/Private/AnimGraphNode_Root.cpp
--- a/Engine/Source/Editor/AnimGraph/Private/AnimGraphNode_Root.cpp
+++ b/Engine/Source/Editor/AnimGraph/Private/AnimGraphNode_Root.cpp
@@ -72,6 +72,11 @@ FText UAnimGraphNode_Root::GetTooltipText() const
 	return LOCTEXT("AnimGraphNodeRoot_Tooltip", "Wire the final animation pose for this graph into this node");
 }
 
+bool UAnimGraphNode_Root::IsPoseWatchable() const
+{
+	return false;
+}
+
 bool UAnimGraphNode_Root::IsSinkNode() const
 {
 	return true;
diff --git a/Engine/Source/Editor/AnimGraph/Public/AnimGraphNode_Base.h b/Engine/Source/Editor/AnimGraph/Public/AnimGraphNode_Base.h
--- a/Engine/Source/Editor/AnimGraph/Public/AnimGraphNode_Base.h
+++ b/Engine/Source/Editor/AnimGraph/Public/AnimGraphNode_Base.h
@@ -274,6 +274,9 @@ public:
 
 	// UAnimGraphNode_Base interface
 
+	// Whether or not you can add a pose watch on this node
+	virtual bool IsPoseWatchable() const { return true; }
+
 	// Gets the menu category this node belongs in
 	virtual FString GetNodeCategory() const;
 
diff --git a/Engine/Source/Editor/AnimGraph/Public/AnimGraphNode_Root.h b/Engine/Source/Editor/AnimGraph/Public/AnimGraphNode_Root.h
--- a/Engine/Source/Editor/AnimGraph/Public/AnimGraphNode_Root.h
+++ b/Engine/Source/Editor/AnimGraph/Public/AnimGraphNode_Root.h
@@ -29,6 +29,7 @@ class UAnimGraphNode_Root : public UAnimGraphNode_Base
 	//~ End UEdGraphNode Interface.
 
 	//~ Begin UAnimGraphNode_Base Interface
+	virtual bool IsPoseWatchable() const override;
 	virtual bool IsSinkNode() const override;
 	virtual void OnProcessDuringCompilation(IAnimBlueprintCompilationContext& InCompilationContext, IAnimBlueprintGeneratedClassCompiledData& OutCompiledData) override;
 	virtual void OnCopyTermDefaultsToDefaultObject(IAnimBlueprintCopyTermDefaultsContext& InCompilationContext, IAnimBlueprintNodeCopyTermDefaultsContext& InPerNodeContext, IAnimBlueprintGeneratedClassCompiledData& OutCompiledData) override;
diff --git a/Engine/Source/Editor/AnimationBlueprintEditor/Private/AnimationBlueprintEditor.cpp b/Engine/Source/Editor/AnimationBlueprintEditor/Private/AnimationBlueprintEditor.cpp
--- a/Engine/Source/Editor/AnimationBlueprintEditor/Private/AnimationBlueprintEditor.cpp
+++ b/Engine/Source/Editor/AnimationBlueprintEditor/Private/AnimationBlueprintEditor.cpp
@@ -549,7 +549,9 @@ void FAnimationBlueprintEditor::CreateDefaultCommands()
 void FAnimationBlueprintEditor::OnCreateGraphEditorCommands(TSharedPtr<FUICommandList> GraphEditorCommandsList)
 {
 	GraphEditorCommandsList->MapAction(FAnimGraphCommands::Get().TogglePoseWatch,
-		FExecuteAction::CreateSP(this, &FAnimationBlueprintEditor::OnTogglePoseWatch));
+		FExecuteAction::CreateSP(this, &FAnimationBlueprintEditor::OnTogglePoseWatch ),
+		FCanExecuteAction::CreateSP(this, &FAnimationBlueprintEditor::CanTogglePoseWatch)
+		);
 
 	GraphEditorCommandsList->MapAction( FAnimGraphCommands::Get().AddBlendListPin,
 		FExecuteAction::CreateSP( this, &FAnimationBlueprintEditor::OnAddPosePin ),
@@ -711,6 +713,30 @@ void FAnimationBlueprintEditor::OnRemovePosePin()
 	}
 }
 
+bool FAnimationBlueprintEditor::CanTogglePoseWatch()
+{
+	const FGraphPanelSelectionSet SelectedNodes = GetSelectedNodes();
+	UAnimBlueprint* AnimBP = GetAnimBlueprint();
+
+	for (FGraphPanelSelectionSet::TConstIterator NodeIt(SelectedNodes); NodeIt; ++NodeIt)
+	{
+		if (UAnimGraphNode_Base* SelectedNode = Cast<UAnimGraphNode_Base>(*NodeIt))
+		{
+			UPoseWatch* ExistingPoseWatch = AnimationEditorUtils::FindPoseWatchForNode(SelectedNode, AnimBP);
+			if (ExistingPoseWatch)
+			{
+				return true;
+			}
+			if (SelectedNode->IsPoseWatchable())
+			{
+				return true;
+			}
+		}
+	}
+
+	return false;
+}
+
 void FAnimationBlueprintEditor::OnTogglePoseWatch()
 {
 	const FGraphPanelSelectionSet SelectedNodes = GetSelectedNodes();
@@ -738,7 +764,7 @@ void FAnimationBlueprintEditor::OnTogglePoseWatch()
 				}
 				AnimationEditorUtils::OnPoseWatchesChanged().Broadcast(AnimBP, ExistingPoseWatch->Node.Get());
 			}
-			else
+			else if (SelectedNode->IsPoseWatchable())
 			{
 				UPoseWatch* NewPoseWatch = AnimationEditorUtils::MakePoseWatchForNode(AnimBP, SelectedNode);
 				AnimationEditorUtils::OnPoseWatchesChanged().Broadcast(AnimBP, NewPoseWatch->Node.Get());
@@ -1854,7 +1880,7 @@ void FAnimationBlueprintEditor::HandlePoseWatchSelectedNodes()
 			{
 				if (SelectionNodes.Contains(Node))
 				{
-					if (!PoseWatch)
+					if (!PoseWatch && GraphNode->IsPoseWatchable())
 					{
 						PoseWatch = AnimationEditorUtils::MakePoseWatchForNode(AnimBP, GraphNode);
 						PoseWatch->SetShouldDeleteOnDeselect(true);
diff --git a/Engine/Source/Editor/AnimationBlueprintEditor/Private/AnimationBlueprintEditor.h b/Engine/Source/Editor/AnimationBlueprintEditor/Private/AnimationBlueprintEditor.h
--- a/Engine/Source/Editor/AnimationBlueprintEditor/Private/AnimationBlueprintEditor.h
+++ b/Engine/Source/Editor/AnimationBlueprintEditor/Private/AnimationBlueprintEditor.h
@@ -208,6 +208,7 @@ protected:
 	//~ End FNotifyHook Interface
 
 	// Toggle pose watch on selected nodes
+	bool CanTogglePoseWatch();
 	void OnTogglePoseWatch();
 
 	void BindCommands();
diff --git a/Engine/Source/Editor/UnrealEd/Private/AnimationEditorUtils.cpp b/Engine/Source/Editor/UnrealEd/Private/AnimationEditorUtils.cpp
--- a/Engine/Source/Editor/UnrealEd/Private/AnimationEditorUtils.cpp
+++ b/Engine/Source/Editor/UnrealEd/Private/AnimationEditorUtils.cpp
@@ -903,6 +903,7 @@ namespace AnimationEditorUtils
 	UPoseWatch* MakePoseWatchForNode(UAnimBlueprint* AnimBlueprint, UEdGraphNode* Node)
 	{
 #if WITH_EDITORONLY_DATA
+		check(CastChecked<UAnimGraphNode_Base>(Node)->IsPoseWatchable());
 		UPoseWatch* NewPoseWatch = NewObject<UPoseWatch>(AnimBlueprint);
 		NewPoseWatch->Node = Node;
 		NewPoseWatch->SetUniqueDefaultLabel();
