Fixed anim pose pin attribute wires not drawing when routing through re-route (knot) nodes

diff --git a/Engine/Source/Editor/AnimationBlueprintEditor/Private/AnimGraphConnectionDrawingPolicy.cpp b/Engine/Source/Editor/AnimationBlueprintEditor/Private/AnimGraphConnectionDrawingPolicy.cpp
--- a/Engine/Source/Editor/AnimationBlueprintEditor/Private/AnimGraphConnectionDrawingPolicy.cpp
+++ b/Engine/Source/Editor/AnimationBlueprintEditor/Private/AnimGraphConnectionDrawingPolicy.cpp
@@ -9,6 +9,7 @@
 #include "AnimGraphAttributes.h"
 #include "Algo/Sort.h"
 #include "AnimationPins/SGraphPinPose.h"
+#include "K2Node_Knot.h"
 
 /////////////////////////////////////////////////////
 // FAnimGraphConnectionDrawingPolicy
@@ -131,119 +132,135 @@ void FAnimGraphConnectionDrawingPolicy::DrawConnection(int32 LayerId, const FVec
 	{
 		if(UAnimationGraphSchema::IsPosePin(Params.AssociatedPin1->PinType))
 		{
-			UAnimGraphNode_Base* Node1 = Cast<UAnimGraphNode_Base>(Params.AssociatedPin1->GetOwningNode());
-			UAnimGraphNode_Base* Node2 = Cast<UAnimGraphNode_Base>(Params.AssociatedPin2->GetOwningNode());
-			if(Node1 && Node2)
+			// If either pin connects to a re-route (knot) node, traverse the link until a relevant pin is found
+			UEdGraphPin* UsePin1 = Params.AssociatedPin1;
+			if (Params.AssociatedPin1->GetOwningNode()->IsA<UK2Node_Knot>())
 			{
-				const TArrayView<const SGraphPinPose::FAttributeInfo>& AdditionalAttributeInfo = PinAttributes.FindRef(Params.AssociatedPin1);
+				UsePin1 = FBlueprintEditorUtils::FindFirstCompilerRelevantLinkedPin(Params.AssociatedPin1);
+			}
+
+			UEdGraphPin* UsePin2 = Params.AssociatedPin2;
+			if (Params.AssociatedPin2->GetOwningNode()->IsA<UK2Node_Knot>())
+			{
+				UsePin2 = FBlueprintEditorUtils::FindFirstCompilerRelevantLinkedPin(Params.AssociatedPin2);
+			}
 
-				if(AdditionalAttributeInfo.Num() > 0)
+			if (UsePin1 && UsePin2)
+			{
+				UAnimGraphNode_Base* Node1 = Cast<UAnimGraphNode_Base>(UsePin1->GetOwningNode());
+				UAnimGraphNode_Base* Node2 = Cast<UAnimGraphNode_Base>(UsePin2->GetOwningNode());
+				if (Node1 && Node2)
 				{
-					const float MaxAttributeWireThickness = 3.0f;
-					const float MinAttributeWireThickness = 1.0f;
-					const float MaxWireGap = 2.0f;
-					const float MinWireGap = 0.5f;
+					const TArrayView<const SGraphPinPose::FAttributeInfo>& AdditionalAttributeInfo = PinAttributes.FindRef(UsePin1);
 
-					// 0.375f is the zoom level before the 'low LOD' cutoff
-					const float ZoomLevelAlpha = FMath::GetMappedRangeValueClamped(TRange<float>(0.375f, 1.0f), TRange<float>(0.0f, 1.0f), PanelZoom);
-					const float AttributeWireThickness = FMath::Lerp(MinAttributeWireThickness, MaxAttributeWireThickness, ZoomLevelAlpha);
-					const float WireGap = FMath::Lerp(MinWireGap, MaxWireGap, ZoomLevelAlpha);
+					if (AdditionalAttributeInfo.Num() > 0)
+					{
+						const float MaxAttributeWireThickness = 3.0f;
+						const float MinAttributeWireThickness = 1.0f;
+						const float MaxWireGap = 2.0f;
+						const float MinWireGap = 0.5f;
 
-					const FVector2D& P0 = Start;
-					const FVector2D& P1 = End;
+						// 0.375f is the zoom level before the 'low LOD' cutoff
+						const float ZoomLevelAlpha = FMath::GetMappedRangeValueClamped(TRange<float>(0.375f, 1.0f), TRange<float>(0.0f, 1.0f), PanelZoom);
+						const float AttributeWireThickness = FMath::Lerp(MinAttributeWireThickness, MaxAttributeWireThickness, ZoomLevelAlpha);
+						const float WireGap = FMath::Lerp(MinWireGap, MaxWireGap, ZoomLevelAlpha);
 
-					const FVector2D SplineTangent = ComputeSplineTangent(P0, P1);
-					const FVector2D P0Tangent = (Params.StartDirection == EGPD_Output) ? SplineTangent : -SplineTangent;
-					const FVector2D P1Tangent = (Params.EndDirection == EGPD_Input) ? SplineTangent : -SplineTangent;
+						const FVector2D& P0 = Start;
+						const FVector2D& P1 = End;
 
-					bCompositeWire = true;
+						const FVector2D SplineTangent = ComputeSplineTangent(P0, P1);
+						const FVector2D P0Tangent = (Params.StartDirection == EGPD_Output) ? SplineTangent : -SplineTangent;
+						const FVector2D P1Tangent = (Params.EndDirection == EGPD_Input) ? SplineTangent : -SplineTangent;
 
-					float TotalThickness = Params.WireThickness;
+						bCompositeWire = true;
 
-					static TArray<float> CachedWireThicknesses;
-					check(CachedWireThicknesses.Num() == 0);	// Cant be called recursively or on multiple threads
-					CachedWireThicknesses.SetNumZeroed(AdditionalAttributeInfo.Num());
+						float TotalThickness = Params.WireThickness;
 
-					for(int32 AttributeInfoIndex = 0; AttributeInfoIndex < AdditionalAttributeInfo.Num(); ++AttributeInfoIndex)
-					{
-						const SGraphPinPose::FAttributeInfo& AttributeInfo = AdditionalAttributeInfo[AttributeInfoIndex];
+						static TArray<float> CachedWireThicknesses;
+						check(CachedWireThicknesses.Num() == 0);	// Cant be called recursively or on multiple threads
+						CachedWireThicknesses.SetNumZeroed(AdditionalAttributeInfo.Num());
 
-						float WireThickness = 0.0f;
-						switch(AttributeInfo.Blend)
-						{
-						case EAnimGraphAttributeBlend::Blendable:
+						for (int32 AttributeInfoIndex = 0; AttributeInfoIndex < AdditionalAttributeInfo.Num(); ++AttributeInfoIndex)
 						{
-							if(UAnimBlueprint* AnimBlueprint = Cast<UAnimBlueprint>(FBlueprintEditorUtils::FindBlueprintForGraph(GraphObj)))
+							const SGraphPinPose::FAttributeInfo& AttributeInfo = AdditionalAttributeInfo[AttributeInfoIndex];
+
+							float WireThickness = 0.0f;
+							switch (AttributeInfo.Blend)
 							{
-								UAnimBlueprintGeneratedClass* AnimBlueprintClass = (UAnimBlueprintGeneratedClass*)(*(AnimBlueprint->GeneratedClass));
-								int32 SourceNodeId = AnimBlueprintClass->GetNodeIndexFromGuid(Node1->NodeGuid);
-								int32 TargetNodeId = AnimBlueprintClass->GetNodeIndexFromGuid(Node2->NodeGuid);
-								if(SourceNodeId != INDEX_NONE && TargetNodeId != INDEX_NONE)
+							case EAnimGraphAttributeBlend::Blendable:
+							{
+								if (UAnimBlueprint* AnimBlueprint = Cast<UAnimBlueprint>(FBlueprintEditorUtils::FindBlueprintForGraph(GraphObj)))
 								{
-									const TArray<FAnimBlueprintDebugData::FAttributeRecord>* LinkAttributes = AnimBlueprintClass->GetAnimBlueprintDebugData().NodeOutputAttributesThisFrame.Find(SourceNodeId);
-									const bool bAttributeUsedInLink = LinkAttributes && LinkAttributes->ContainsByPredicate(
-										[&AttributeInfo, TargetNodeId](const FAnimBlueprintDebugData::FAttributeRecord& InRecord)
-										{
-											return InRecord.Attribute == AttributeInfo.Attribute && InRecord.OtherNode == TargetNodeId; 
-										});
-
-
-									WireThickness = bAttributeUsedInLink ? AttributeWireThickness : 0.0f;
+									UAnimBlueprintGeneratedClass* AnimBlueprintClass = (UAnimBlueprintGeneratedClass*)(*(AnimBlueprint->GeneratedClass));
+									int32 SourceNodeId = AnimBlueprintClass->GetNodeIndexFromGuid(Node1->NodeGuid);
+									int32 TargetNodeId = AnimBlueprintClass->GetNodeIndexFromGuid(Node2->NodeGuid);
+									if (SourceNodeId != INDEX_NONE && TargetNodeId != INDEX_NONE)
+									{
+										const TArray<FAnimBlueprintDebugData::FAttributeRecord>* LinkAttributes = AnimBlueprintClass->GetAnimBlueprintDebugData().NodeOutputAttributesThisFrame.Find(SourceNodeId);
+										const bool bAttributeUsedInLink = LinkAttributes && LinkAttributes->ContainsByPredicate(
+											[&AttributeInfo, TargetNodeId](const FAnimBlueprintDebugData::FAttributeRecord& InRecord)
+											{
+												return InRecord.Attribute == AttributeInfo.Attribute && InRecord.OtherNode == TargetNodeId;
+											});
+
+
+										WireThickness = bAttributeUsedInLink ? AttributeWireThickness : 0.0f;
+									}
 								}
+								break;
+							}
+							case EAnimGraphAttributeBlend::NonBlendable:
+								WireThickness = AttributeWireThickness;
+								break;
 							}
-							break;
-						}
-						case EAnimGraphAttributeBlend::NonBlendable:
-							WireThickness = AttributeWireThickness;
-							break;
-						}
 
-						CachedWireThicknesses[AttributeInfoIndex] = WireThickness;
-						TotalThickness += WireThickness != 0.0f ? (WireThickness + WireGap) : 0.0f;
-					}
+							CachedWireThicknesses[AttributeInfoIndex] = WireThickness;
+							TotalThickness += WireThickness != 0.0f ? (WireThickness + WireGap) : 0.0f;
+						}
 
-					const float InitialOffset = TotalThickness * 0.5f;
-					FVector2D SubWireP0 = P0;
-					SubWireP0.Y += InitialOffset;
-					FVector2D SubWireP1 = P1;
-					SubWireP1.Y += InitialOffset;
+						const float InitialOffset = TotalThickness * 0.5f;
+						FVector2D SubWireP0 = P0;
+						SubWireP0.Y += InitialOffset;
+						FVector2D SubWireP1 = P1;
+						SubWireP1.Y += InitialOffset;
 
-					// Draw in reverse order to get pose wires appearing on top
-					for(int32 AttributeInfoIndex = AdditionalAttributeInfo.Num() - 1; AttributeInfoIndex >= 0; --AttributeInfoIndex)
-					{
-						float Thickness = CachedWireThicknesses[AttributeInfoIndex];
-						if(Thickness > 0.0f)
+						// Draw in reverse order to get pose wires appearing on top
+						for (int32 AttributeInfoIndex = AdditionalAttributeInfo.Num() - 1; AttributeInfoIndex >= 0; --AttributeInfoIndex)
 						{
-							const SGraphPinPose::FAttributeInfo& AttributeInfo = AdditionalAttributeInfo[AttributeInfoIndex];
-							FLinearColor Color = AttributeInfo.Color;
-
-							if(HoveredPins.Num() > 0)
+							float Thickness = CachedWireThicknesses[AttributeInfoIndex];
+							if (Thickness > 0.0f)
 							{
-								ApplyHoverDeemphasis(Params.AssociatedPin1, Params.AssociatedPin2, Thickness, Color);
-							}
+								const SGraphPinPose::FAttributeInfo& AttributeInfo = AdditionalAttributeInfo[AttributeInfoIndex];
+								FLinearColor Color = AttributeInfo.Color;
+
+								if (HoveredPins.Num() > 0)
+								{
+									ApplyHoverDeemphasis(Params.AssociatedPin1, Params.AssociatedPin2, Thickness, Color);
+								}
 
-							SubWireP0.Y -= Thickness + WireGap;
-							SubWireP1.Y -= Thickness + WireGap;
-
-							// Draw the spline itself
-							FSlateDrawElement::MakeDrawSpaceSpline(
-								DrawElementsList,
-								LayerId,
-								SubWireP0, P0Tangent,
-								SubWireP1, P1Tangent,
-								Thickness,
-								ESlateDrawEffect::None,
-								Color
-							);
+								SubWireP0.Y -= Thickness + WireGap;
+								SubWireP1.Y -= Thickness + WireGap;
+
+								// Draw the spline itself
+								FSlateDrawElement::MakeDrawSpaceSpline(
+									DrawElementsList,
+									LayerId,
+									SubWireP0, P0Tangent,
+									SubWireP1, P1Tangent,
+									Thickness,
+									ESlateDrawEffect::None,
+									Color
+								);
+							}
 						}
-					}
 
-					SubWireP0.Y -= Params.WireThickness + WireGap;
-					SubWireP1.Y -= Params.WireThickness + WireGap;
+						SubWireP0.Y -= Params.WireThickness + WireGap;
+						SubWireP1.Y -= Params.WireThickness + WireGap;
 
-					FKismetConnectionDrawingPolicy::DrawConnection(LayerId, SubWireP0, SubWireP1, Params);
+						FKismetConnectionDrawingPolicy::DrawConnection(LayerId, SubWireP0, SubWireP1, Params);
 
-					CachedWireThicknesses.Reset();
+						CachedWireThicknesses.Reset();
+					}
 				}
 			}
 		}
