Fixed create comment action not working in all animbp graph schemas

diff --git a/Engine/Source/Editor/AnimationBlueprintEditor/Private/AnimationBlueprintEditor.cpp b/Engine/Source/Editor/AnimationBlueprintEditor/Private/AnimationBlueprintEditor.cpp
--- a/Engine/Source/Editor/AnimationBlueprintEditor/Private/AnimationBlueprintEditor.cpp
+++ b/Engine/Source/Editor/AnimationBlueprintEditor/Private/AnimationBlueprintEditor.cpp
@@ -77,6 +77,7 @@
 #include "AnimationBlueprintEditorSettings.h"
 #include "Framework/MultiBox/MultiBoxBuilder.h"
 #include "EdGraphNode_Comment.h"
+#include "EdGraphSchema_K2_Actions.h"
 #include "AnimStateNodeBase.h"
 #include "AnimStateEntryNode.h"
 #include "PersonaUtils.h"
@@ -1705,6 +1706,19 @@ void FAnimationBlueprintEditor::CreateEditorModeManager()
 	EditorModeManager = MakeShareable(FModuleManager::LoadModuleChecked<FPersonaModule>("Persona").CreatePersonaEditorModeManager());
 }
 
+void FAnimationBlueprintEditor::OnCreateComment()
+{
+	TSharedPtr<SGraphEditor> GraphEditor = FocusedGraphEdPtr.Pin();
+	if (GraphEditor.IsValid())
+	{
+		if (UEdGraph* Graph = GraphEditor->GetCurrentGraph())
+		{
+			FEdGraphSchemaAction_K2AddComment CommentAction;
+			CommentAction.PerformAction(Graph, nullptr, GraphEditor->GetPasteLocation());
+		}
+	}
+}
+
 void FAnimationBlueprintEditor::JumpToHyperlink(const UObject* ObjectReference, bool bRequestRename)
 {
 	if(const UBlendSpaceGraph* BlendSpaceGraph = Cast<UBlendSpaceGraph>(ObjectReference))
diff --git a/Engine/Source/Editor/AnimationBlueprintEditor/Private/AnimationBlueprintEditor.h b/Engine/Source/Editor/AnimationBlueprintEditor/Private/AnimationBlueprintEditor.h
--- a/Engine/Source/Editor/AnimationBlueprintEditor/Private/AnimationBlueprintEditor.h
+++ b/Engine/Source/Editor/AnimationBlueprintEditor/Private/AnimationBlueprintEditor.h
@@ -196,6 +196,7 @@ protected:
 	virtual FText GetGraphDecorationString(UEdGraph* InGraph) const override;
 	virtual void OnBlueprintChangedImpl(UBlueprint* InBlueprint, bool bIsJustBeingCompiled = false) override;
 	virtual void CreateEditorModeManager() override;
+	virtual void OnCreateComment() override;
 	//~ End FBlueprintEditor Interface
 
 	//~ Begin FEditorUndoClient Interface
