Added CreateNodeBody to SGraphNodeStateMachineInstance in order to show Node Tag

diff --git a/Engine/Source/Editor/AnimationBlueprintEditor/Private/AnimationNodes/SGraphNodeStateMachineInstance.cpp b/Engine/Source/Editor/AnimationBlueprintEditor/Private/AnimationNodes/SGraphNodeStateMachineInstance.cpp
--- a/Engine/Source/Editor/AnimationBlueprintEditor/Private/AnimationNodes/SGraphNodeStateMachineInstance.cpp
+++ b/Engine/Source/Editor/AnimationBlueprintEditor/Private/AnimationNodes/SGraphNodeStateMachineInstance.cpp
@@ -1,6 +1,7 @@
 // Copyright Epic Games, Inc. All Rights Reserved.
 
 #include "AnimationNodes/SGraphNodeStateMachineInstance.h"
+#include "AnimationNodes/SAnimationGraphNode.h"
 #include "AnimationStateMachineGraph.h"
 #include "AnimGraphNode_StateMachineBase.h"
 #include "Engine/PoseWatch.h"
@@ -50,4 +51,30 @@ TArray<FOverlayWidgetInfo> SGraphNodeStateMachineInstance::GetOverlayWidgets(boo
 	return Widgets;
 }
 
+TSharedRef<SWidget> SGraphNodeStateMachineInstance::CreateNodeBody()
+{
+	TSharedRef<SWidget> NodeBody = SGraphNodeK2Composite::CreateNodeBody();
+
+	UAnimGraphNode_StateMachineBase* StateMachineNode = CastChecked<UAnimGraphNode_StateMachineBase>(GraphNode);
+
+	auto UseLowDetailNode = [this]()
+	{
+		return GetCurrentLOD() <= EGraphRenderingLOD::LowDetail;
+	};
+
+	return SNew(SVerticalBox)
+		+ SVerticalBox::Slot()
+			.AutoHeight()
+			[
+				NodeBody
+			]
+		+ SVerticalBox::Slot()
+			.AutoHeight()
+			.HAlign(HAlign_Right)
+			.Padding(4.0f, 2.0f, 4.0f, 2.0f)
+			[
+				SAnimationGraphNode::CreateNodeTagWidget(StateMachineNode, MakeAttributeLambda(UseLowDetailNode))
+			];
+}
+
 #undef LOCTEXT_NAMESPACE
\ No newline at end of file
diff --git a/Engine/Source/Editor/AnimationBlueprintEditor/Private/AnimationNodes/SGraphNodeStateMachineInstance.h b/Engine/Source/Editor/AnimationBlueprintEditor/Private/AnimationNodes/SGraphNodeStateMachineInstance.h
--- a/Engine/Source/Editor/AnimationBlueprintEditor/Private/AnimationNodes/SGraphNodeStateMachineInstance.h
+++ b/Engine/Source/Editor/AnimationBlueprintEditor/Private/AnimationNodes/SGraphNodeStateMachineInstance.h
@@ -21,6 +21,7 @@ protected:
 	// SGraphNodeK2Composite interface
 	virtual UEdGraph* GetInnerGraph() const override;
 	virtual TArray<FOverlayWidgetInfo> GetOverlayWidgets(bool bSelected, const FVector2D& WidgetSize) const override;
+	virtual TSharedRef<SWidget> CreateNodeBody() override;
 	// End of SGraphNodeK2Composite interface
 
 private:
