[Bug Fix] Crash in the diff tool when trying to review an anim blueprint

diff --git a/Engine/Source/Editor/AnimationBlueprintEditor/Private/AnimationPins/SGraphPinPose.cpp b/Engine/Source/Editor/AnimationBlueprintEditor/Private/AnimationPins/SGraphPinPose.cpp
--- a/Engine/Source/Editor/AnimationBlueprintEditor/Private/AnimationPins/SGraphPinPose.cpp
+++ b/Engine/Source/Editor/AnimationBlueprintEditor/Private/AnimationPins/SGraphPinPose.cpp
@@ -210,7 +210,7 @@ void SGraphPinPose::ReconfigureWidgetForAttributes()
 			AnimGraphNode->GetOutputLinkAttributes(NodeAttributes);
 
 			// Unlinked pins display attributes if they are inputs and the node takes them as inputs
-			if(GraphPinObj->LinkedTo.Num() == 0)
+			if(GraphPinObj->LinkedTo.Num() == 0 || AnimBlueprintClass == nullptr)
 			{
 				AddAttributes(TArrayView<const FName>(), PinAttributes, TArrayView<const FName>());
 			}
