Ensure that LinkedAnimGraph node name, and graph in MyBlueprint tab reflect the currently previewed instance and its state (e.g. PIE Anim instance, with dynamically linked graph)

diff --git a/Engine/Source/Editor/AnimGraph/Private/AnimGraphNode_LinkedAnimGraph.cpp b/Engine/Source/Editor/AnimGraph/Private/AnimGraphNode_LinkedAnimGraph.cpp
--- a/Engine/Source/Editor/AnimGraph/Private/AnimGraphNode_LinkedAnimGraph.cpp
+++ b/Engine/Source/Editor/AnimGraph/Private/AnimGraphNode_LinkedAnimGraph.cpp
@@ -190,4 +190,28 @@ void UAnimGraphNode_LinkedAnimGraph::Serialize(FArchive& Ar)
 	}
 }
 
+FAnimNode_CustomProperty* UAnimGraphNode_LinkedAnimGraph::GetCustomPropertyNode()
+{
+	FAnimNode_CustomProperty* const RuntimeCustomPropertyNode = GetDebuggedAnimNode<FAnimNode_CustomProperty>();
+	return RuntimeCustomPropertyNode ? RuntimeCustomPropertyNode : &Node;
+}
+
+const FAnimNode_CustomProperty* UAnimGraphNode_LinkedAnimGraph::GetCustomPropertyNode() const
+{
+	const FAnimNode_CustomProperty* const RuntimeCustomPropertyNode = GetDebuggedAnimNode<FAnimNode_CustomProperty>();
+	return RuntimeCustomPropertyNode ? RuntimeCustomPropertyNode : &Node;
+}
+
+FAnimNode_LinkedAnimGraph* UAnimGraphNode_LinkedAnimGraph::GetLinkedAnimGraphNode()
+{
+	FAnimNode_LinkedAnimGraph* const RuntimeLinkedAnimGraphNode = GetDebuggedAnimNode<FAnimNode_LinkedAnimGraph>();
+	return RuntimeLinkedAnimGraphNode ? RuntimeLinkedAnimGraphNode : &Node;
+}
+
+const FAnimNode_LinkedAnimGraph* UAnimGraphNode_LinkedAnimGraph::GetLinkedAnimGraphNode() const
+{
+	const FAnimNode_LinkedAnimGraph* const RuntimeLinkedAnimGraphNode = GetDebuggedAnimNode<FAnimNode_LinkedAnimGraph>();
+	return RuntimeLinkedAnimGraphNode ? RuntimeLinkedAnimGraphNode : &Node;
+}
+
 #undef LOCTEXT_NAMESPACE
\ No newline at end of file
diff --git a/Engine/Source/Editor/AnimGraph/Public/AnimGraphNode_LinkedAnimGraph.h b/Engine/Source/Editor/AnimGraph/Public/AnimGraphNode_LinkedAnimGraph.h
--- a/Engine/Source/Editor/AnimGraph/Public/AnimGraphNode_LinkedAnimGraph.h
+++ b/Engine/Source/Editor/AnimGraph/Public/AnimGraphNode_LinkedAnimGraph.h
@@ -31,12 +31,12 @@ public:
 	virtual bool IsActionFilteredOut(class FBlueprintActionFilter const& Filter) override;
 	
 	// Begin UAnimGraphNode_CustomProperty
-	virtual FAnimNode_CustomProperty* GetCustomPropertyNode() override { return &Node;  }
-	virtual const FAnimNode_CustomProperty* GetCustomPropertyNode() const override { return &Node; }
+	virtual FAnimNode_CustomProperty* GetCustomPropertyNode() override;
+	virtual const FAnimNode_CustomProperty* GetCustomPropertyNode() const override;
 
 	// Begin UAnimGraphNode_LinkedAnimGraphBase
-	virtual FAnimNode_LinkedAnimGraph* GetLinkedAnimGraphNode() override { return &Node; }
-	virtual const FAnimNode_LinkedAnimGraph* GetLinkedAnimGraphNode() const override { return &Node; }
+	virtual FAnimNode_LinkedAnimGraph* GetLinkedAnimGraphNode() override;
+	virtual const FAnimNode_LinkedAnimGraph* GetLinkedAnimGraphNode() const override;
 
 	// IK2Node_ExternalGraphInterface interface
 	virtual TArray<UEdGraph*> GetExternalGraphs() const override;
diff --git a/Engine/Source/Editor/AnimationBlueprintEditor/Private/AnimationGraphFactory.cpp b/Engine/Source/Editor/AnimationBlueprintEditor/Private/AnimationGraphFactory.cpp
--- a/Engine/Source/Editor/AnimationBlueprintEditor/Private/AnimationGraphFactory.cpp
+++ b/Engine/Source/Editor/AnimationBlueprintEditor/Private/AnimationGraphFactory.cpp
@@ -29,6 +29,7 @@
 #include "AnimationNodes/SGraphNodeStateMachineInstance.h"
 #include "AnimationNodes/SGraphNodeLayeredBoneBlend.h"
 #include "AnimationNodes/SGraphNodeBlendSpacePlayer.h"
+#include "AnimationNodes/SGraphNodeLinkedLayer.h"
 #include "AnimationPins/SGraphPinPose.h"
 
 #include "AnimGraphConnectionDrawingPolicy.h"
@@ -36,6 +37,7 @@
 
 #include "KismetPins/SGraphPinExec.h"
 #include "AnimGraphNode_BlendSpaceGraph.h"
+#include "AnimGraphNode_LinkedAnimGraph.h"
 #include "AnimationNodes/SGraphNodeBlendSpaceGraph.h"
 #include "K2Node_AnimNodeReference.h"
 #include "AnimationNodes/SAnimNodeReference.h"
@@ -68,6 +70,10 @@ TSharedPtr<class SGraphNode> FAnimationGraphNodeFactory::CreateNode(class UEdGra
 		{
 			return SNew(SGraphNodeBlendSpaceGraph, BlendSpaceGraph);
 		}
+		else if (UAnimGraphNode_LinkedAnimGraph* LinkedAnimLayer = Cast<UAnimGraphNode_LinkedAnimGraph>(InNode))
+		{
+			return SNew(SGraphNodeLinkedLayer, LinkedAnimLayer);
+		}
 		else
 		{
 			return SNew(SAnimationGraphNode, BaseAnimNode);
diff --git a/Engine/Source/Editor/AnimationBlueprintEditor/Private/AnimationNodes/SGraphNodeLinkedLayer.cpp b/Engine/Source/Editor/AnimationBlueprintEditor/Private/AnimationNodes/SGraphNodeLinkedLayer.cpp
new file mode 100644
--- /dev/null
+++ b/Engine/Source/Editor/AnimationBlueprintEditor/Private/AnimationNodes/SGraphNodeLinkedLayer.cpp
@@ -0,0 +1,61 @@
+// Copyright Epic Games, Inc. All Rights Reserved.
+
+#include "AnimationNodes/SGraphNodeLinkedLayer.h"
+#include "AnimGraphNode_Base.h"
+#include "AnimGraphNode_LinkedAnimGraph.h"
+#include "BlueprintEditorModule.h"
+#include "Kismet2/BlueprintEditorUtils.h"
+#include "SLevelOfDetailBranchNode.h"
+#include "Widgets/Layout/SSpacer.h"
+
+void SGraphNodeLinkedLayer::Construct(const FArguments& InArgs, UAnimGraphNode_Base* InNode)
+{
+	this->GraphNode = InNode;
+	this->UpdateGraphNode();
+
+	CachedTargetName = NAME_None;
+
+	SAnimationGraphNode::Construct(SAnimationGraphNode::FArguments(), InNode);
+
+	RegisterActiveTimer(1.0f / 10.0f, FWidgetActiveTimerDelegate::CreateLambda([this](double InCurrentTime, float InDeltaTime)
+	{
+		UpdateNodeLabel();
+		return EActiveTimerReturnType::Continue;
+	}));
+}
+
+void SGraphNodeLinkedLayer::UpdateNodeLabel()
+{
+	if (UAnimGraphNode_LinkedAnimGraph* VisualLinkedAnimLayer = Cast<UAnimGraphNode_LinkedAnimGraph>(GraphNode))
+	{
+		FName FoundName = NAME_None;
+		
+		if (UAnimBlueprint* AnimBlueprint = Cast<UAnimBlueprint>(FBlueprintEditorUtils::FindBlueprintForNode(GraphNode)))
+		{
+			if (const UAnimInstance* InstanceBeingDebugged = Cast<const UAnimInstance>(AnimBlueprint->GetObjectBeingDebugged()))
+			{
+				USkeletalMeshComponent* Component = InstanceBeingDebugged->GetSkelMeshComponent();
+				if (Component != nullptr)
+				{
+					if (const FAnimNode_LinkedAnimGraph* AnimGraphNode = static_cast<FAnimNode_LinkedAnimGraph*>(VisualLinkedAnimLayer->FindDebugAnimNode(Component)))
+					{
+						if (const UAnimInstance* TargetInstance = AnimGraphNode->GetTargetInstance<UAnimInstance>())
+						{
+							FoundName = TargetInstance->GetFName();
+
+							if (FoundName != CachedTargetName)
+							{
+								VisualLinkedAnimLayer->OnNodeTitleChangedEvent().Broadcast();
+								CachedTargetName = FoundName;
+								if (IAssetEditorInstance* EditorInstance = GEditor->GetEditorSubsystem<UAssetEditorSubsystem>()->FindEditorForAsset(AnimBlueprint, false))
+								{
+									static_cast<IBlueprintEditor*>(EditorInstance)->RefreshMyBlueprint();						
+								}
+							}
+						}
+					}
+				}
+			}	
+		}
+	}
+}
diff --git a/Engine/Source/Editor/AnimationBlueprintEditor/Private/AnimationNodes/SGraphNodeLinkedLayer.h b/Engine/Source/Editor/AnimationBlueprintEditor/Private/AnimationNodes/SGraphNodeLinkedLayer.h
new file mode 100644
--- /dev/null
+++ b/Engine/Source/Editor/AnimationBlueprintEditor/Private/AnimationNodes/SGraphNodeLinkedLayer.h
@@ -0,0 +1,26 @@
+// Copyright Epic Games, Inc. All Rights Reserved.
+
+#pragma once
+
+#include "CoreMinimal.h"
+#include "Layout/Visibility.h"
+#include "Widgets/DeclarativeSyntaxSupport.h"
+#include "SNodePanel.h"
+#include "AnimationNodes/SAnimationGraphNode.h"
+#include "Animation/BlendSpace.h"
+
+class UAnimGraphNode_Base;
+
+class SGraphNodeLinkedLayer : public SAnimationGraphNode
+{
+public:
+	SLATE_BEGIN_ARGS(SGraphNodeLinkedLayer){}
+	SLATE_END_ARGS()
+	
+	void Construct(const FArguments& InArgs, UAnimGraphNode_Base* InNode);
+
+protected:
+	void UpdateNodeLabel();
+
+	FName CachedTargetName;
+};
diff --git a/Engine/Source/Editor/Kismet/Private/SBlueprintEditorSelectedDebugObjectWidget.cpp b/Engine/Source/Editor/Kismet/Private/SBlueprintEditorSelectedDebugObjectWidget.cpp
--- a/Engine/Source/Editor/Kismet/Private/SBlueprintEditorSelectedDebugObjectWidget.cpp
+++ b/Engine/Source/Editor/Kismet/Private/SBlueprintEditorSelectedDebugObjectWidget.cpp
@@ -549,6 +549,11 @@ void SBlueprintEditorSelectedDebugObjectWidget::DebugObjectSelectionChanged(TSha
 	{
 		UObject* DebugObj = NewSelection->ObjectPtr.Get();
 		GetBlueprintObj()->SetObjectBeingDebugged(DebugObj);
+
+		if (TSharedPtr<FBlueprintEditor> SharedBlueprintEditor = BlueprintEditor.Pin())
+		{
+			SharedBlueprintEditor->RefreshMyBlueprint();
+		}
 		
 		LastObjectObserved = DebugObj;
 	}
