Add a context menu option to unexpose all unused pins in LinkedAnimGraph and LinkedAnimLayer nodes in animation blueprints
This is to be able to quickly hide all pins on linked graph nodes based on animation instance that expose a lot of properties

#tests PIE

diff --git a/Engine/Source/Editor/AnimGraph/Private/AnimGraphCommands.cpp b/Engine/Source/Editor/AnimGraph/Private/AnimGraphCommands.cpp
--- a/Engine/Source/Editor/AnimGraph/Private/AnimGraphCommands.cpp
+++ b/Engine/Source/Editor/AnimGraph/Private/AnimGraphCommands.cpp
@@ -10,6 +10,7 @@
 void FAnimGraphCommands::RegisterCommands()
 {
 	UI_COMMAND(TogglePoseWatch, "Toggle Pose Watch", "Toggle pose watching on this node", EUserInterfaceActionType::Button, FInputChord());
+	UI_COMMAND(HideUnboundPropertyPins, "Hide Unbound/Unset Property Pins", "Unexpose all unbound/unset property pins from this node", EUserInterfaceActionType::Button, FInputChord());
 
 	UI_COMMAND( SelectBone, "Select Bone", "Assign or change the bone for skeletal controls", EUserInterfaceActionType::Button, FInputChord() )
 	UI_COMMAND( AddBlendListPin, "Add Blend Pin", "Add blend pin to blend list", EUserInterfaceActionType::Button, FInputChord() )
diff --git a/Engine/Source/Editor/AnimGraph/Private/AnimGraphNode_Base.cpp b/Engine/Source/Editor/AnimGraph/Private/AnimGraphNode_Base.cpp
--- a/Engine/Source/Editor/AnimGraph/Private/AnimGraphNode_Base.cpp
+++ b/Engine/Source/Editor/AnimGraph/Private/AnimGraphNode_Base.cpp
@@ -829,6 +829,16 @@ bool UAnimGraphNode_Base::IsPinExposedAndBound(const FString& InPinName, const E
 	return Pin != nullptr && Pin->LinkedTo.Num() == 0 && PropertyBindings.Find(Pin->GetFName()) != nullptr;
 }
 
+bool UAnimGraphNode_Base::IsPinUnlinkedUnboundAndUnset(const FString& InPinName, const EEdGraphPinDirection InDirection) const
+{
+	UEdGraphPin* Pin = FindPin(InPinName, InDirection);
+	if (Pin)
+	{
+		return Pin->LinkedTo.Num() == 0 && PropertyBindings.Find(Pin->GetFName()) == nullptr && Pin->DefaultValue == Pin->AutogeneratedDefaultValue;
+	}
+	return false;
+}
+
 bool UAnimGraphNode_Base::IsPinBindable(const UEdGraphPin* InPin) const
 {
 	if(const FProperty* PinProperty = GetPinProperty(InPin->GetFName()))
diff --git a/Engine/Source/Editor/AnimGraph/Private/AnimationGraphSchema.cpp b/Engine/Source/Editor/AnimGraph/Private/AnimationGraphSchema.cpp
--- a/Engine/Source/Editor/AnimGraph/Private/AnimationGraphSchema.cpp
+++ b/Engine/Source/Editor/AnimGraph/Private/AnimationGraphSchema.cpp
@@ -36,6 +36,7 @@
 #include "Animation/AnimMontage.h"
 #include "AnimGraphNode_LinkedInputPose.h"
 #include "AnimGraphNode_LinkedAnimLayer.h"
+#include "AnimGraphNode_LinkedAnimGraphBase.h"
 #include "AnimGraphNode_RigidBody.h"
 #include "AnimationBlendSpaceSampleGraph.h"
 #include "GraphEditorDragDropAction.h"
@@ -766,6 +767,7 @@ void UAnimationGraphSchema::GetContextMenuActions(UToolMenu* Menu, UGraphNodeCon
 			// Node contextual actions
 			FToolMenuSection& Section = Menu->AddSection("AnimGraphSchemaNodeActions", LOCTEXT("AnimNodeActionsMenuHeader", "Anim Node Actions"));
 			Section.AddMenuEntry(FAnimGraphCommands::Get().TogglePoseWatch);
+			Section.AddMenuEntry(FAnimGraphCommands::Get().HideUnboundPropertyPins);
 		}
 
 		if(Context->Pin && !IsPosePin(Context->Pin->PinType))
@@ -780,6 +782,23 @@ void UAnimationGraphSchema::GetContextMenuActions(UToolMenu* Menu, UGraphNodeCon
 	}
 }
 
+void UAnimationGraphSchema::HideUnboundPropertyPins(UAnimGraphNode_LinkedAnimGraphBase* Node)
+{
+	TArrayView<FOptionalPinFromProperty> OptionalPins = Node->CustomPinProperties;
+
+	for (FOptionalPinFromProperty& OptionalPin : OptionalPins)
+	{
+		FName PropertyName;
+		FProperty* Property = nullptr;
+		int32 PinIndex = 0;
+		Node->GetPinBindingInfo(OptionalPin.PropertyName, PropertyName, Property, PinIndex);
+		if (Node->IsPinUnlinkedUnboundAndUnset(OptionalPin.PropertyName.ToString(), EGPD_Input))
+		{
+			Node->SetCustomPinVisibility(false, PinIndex);
+		}
+	}
+}
+
 TSharedPtr<SWidget> UAnimationGraphSchema::MakeBindingWidgetForPin(const TArray<UAnimGraphNode_Base*>& InAnimGraphNodes, FName InPinName, bool bInOnGraphNode, TAttribute<bool> bInIsEnabled)
 {
 	const UAnimGraphNode_Base* FirstNode = InAnimGraphNodes[0];
diff --git a/Engine/Source/Editor/AnimGraph/Public/AnimGraphCommands.h b/Engine/Source/Editor/AnimGraph/Public/AnimGraphCommands.h
--- a/Engine/Source/Editor/AnimGraph/Public/AnimGraphCommands.h
+++ b/Engine/Source/Editor/AnimGraph/Public/AnimGraphCommands.h
@@ -29,6 +29,9 @@ public:
 	// Toggle pose watching for a given node
 	TSharedPtr<FUICommandInfo> TogglePoseWatch;
 
+	// Hide unbound property pins
+	TSharedPtr<FUICommandInfo> HideUnboundPropertyPins;
+
 	// SkeletalControl specific commands
 	TSharedPtr< FUICommandInfo > SelectBone;
 	// Blend list options
diff --git a/Engine/Source/Editor/AnimGraph/Public/AnimGraphNode_Base.h b/Engine/Source/Editor/AnimGraph/Public/AnimGraphNode_Base.h
--- a/Engine/Source/Editor/AnimGraph/Public/AnimGraphNode_Base.h
+++ b/Engine/Source/Editor/AnimGraph/Public/AnimGraphNode_Base.h
@@ -426,6 +426,14 @@ public:
 	 */
 	bool IsPinExposedAndBound(const FString& InPinName, const EEdGraphPinDirection InDirection = EGPD_MAX) const;
 
+	/**
+	 * Helper function to check whether a pin is not linked, not bound via property access and still has its default value
+	 * @param	InPinName		The name of the pin @see UEdGraphNode::FindPin
+	 * @param	InPinDirection	The direction of the pin we are looking for. If this is EGPD_MAX, all directions are considered
+	 * @return true if the pin is unlinked, unbound and still has its default value
+	 */
+	bool IsPinUnlinkedUnboundAndUnset(const FString& InPinName, const EEdGraphPinDirection InDirection) const;
+
 	// Event that is broadcast to inform observers that the node title has changed
 	// The default SAnimationGraphNode uses this to invalidate cached node title text
 	DECLARE_EVENT(UAnimGraphNode_Base, FOnNodeTitleChangedEvent);
diff --git a/Engine/Source/Editor/AnimGraph/Public/AnimationGraphSchema.h b/Engine/Source/Editor/AnimGraph/Public/AnimationGraphSchema.h
--- a/Engine/Source/Editor/AnimGraph/Public/AnimationGraphSchema.h
+++ b/Engine/Source/Editor/AnimGraph/Public/AnimationGraphSchema.h
@@ -16,6 +16,7 @@ class UAnimBlueprint;
 class UPhysicsAsset;
 struct FBPInterfaceDescription;
 class UAnimGraphNode_Base;
+class UAnimGraphNode_LinkedAnimGraphBase;
 
 UCLASS(MinimalAPI)
 class UAnimationGraphSchema : public UEdGraphSchema_K2
@@ -119,6 +120,10 @@ class UAnimationGraphSchema : public UEdGraphSchema_K2
 
 	/** Create a binding widget for the specified named pin on the specified anim graph nodes */
 	ANIMGRAPH_API static TSharedPtr<SWidget> MakeBindingWidgetForPin(const TArray<UAnimGraphNode_Base*>& InAnimGraphNodes, FName InPinName, bool bInOnGraphNode, TAttribute<bool> bInIsEnabled);
+	
+	/** Unexpose pins that are unused on a linked anim graph node */
+	ANIMGRAPH_API static void HideUnboundPropertyPins(UAnimGraphNode_LinkedAnimGraphBase* Node);
+
 };
 
 
diff --git a/Engine/Source/Editor/AnimationBlueprintEditor/Private/AnimationBlueprintEditor.cpp b/Engine/Source/Editor/AnimationBlueprintEditor/Private/AnimationBlueprintEditor.cpp
--- a/Engine/Source/Editor/AnimationBlueprintEditor/Private/AnimationBlueprintEditor.cpp
+++ b/Engine/Source/Editor/AnimationBlueprintEditor/Private/AnimationBlueprintEditor.cpp
@@ -12,6 +12,7 @@
 #include "AnimGraphNode_BlendSpaceGraph.h"
 #include "AnimGraphNode_BlendSpacePlayer.h"
 #include "AnimGraphNode_LayeredBoneBlend.h"
+#include "AnimGraphNode_LinkedAnimGraphBase.h"
 #include "AnimGraphNode_MultiWayBlend.h"
 #include "AnimGraphNode_PoseBlendNode.h"
 #include "AnimGraphNode_PoseByName.h"
@@ -46,6 +47,7 @@
 #include "AnimationBlueprintTemplateEditorMode.h"
 #include "AnimationEditorUtils.h"
 #include "AnimationGraph.h"
+#include "AnimationGraphSchema.h"
 #include "AssetRegistry/AssetData.h"
 #include "BlendSpaceDocumentTabFactory.h"
 #include "BlendSpaceGraph.h"
@@ -625,9 +627,14 @@ void FAnimationBlueprintEditor::CreateDefaultCommands()
 void FAnimationBlueprintEditor::OnCreateGraphEditorCommands(TSharedPtr<FUICommandList> GraphEditorCommandsList)
 {
 	GraphEditorCommandsList->MapAction(FAnimGraphCommands::Get().TogglePoseWatch,
-		FExecuteAction::CreateSP(this, &FAnimationBlueprintEditor::OnTogglePoseWatch ),
+		FExecuteAction::CreateSP(this, &FAnimationBlueprintEditor::OnTogglePoseWatch),
 		FCanExecuteAction::CreateSP(this, &FAnimationBlueprintEditor::CanTogglePoseWatch)
-		);
+	);
+
+	GraphEditorCommandsList->MapAction(FAnimGraphCommands::Get().HideUnboundPropertyPins,
+		FExecuteAction::CreateSP(this, &FAnimationBlueprintEditor::OnHideUnboundPropertyPins),
+		FCanExecuteAction::CreateSP(this, &FAnimationBlueprintEditor::CanHideUnboundPropertyPins)
+	);
 
 	GraphEditorCommandsList->MapAction( FAnimGraphCommands::Get().AddBlendListPin,
 		FExecuteAction::CreateSP( this, &FAnimationBlueprintEditor::OnAddPosePin ),
@@ -813,6 +820,20 @@ bool FAnimationBlueprintEditor::CanTogglePoseWatch()
 	return false;
 }
 
+bool FAnimationBlueprintEditor::CanHideUnboundPropertyPins()
+{
+	const FGraphPanelSelectionSet SelectedNodes = GetSelectedNodes();
+	for (FGraphPanelSelectionSet::TConstIterator NodeIt(SelectedNodes); NodeIt; ++NodeIt)
+	{
+		if (UAnimGraphNode_LinkedAnimGraphBase* SelectedNode = Cast<UAnimGraphNode_LinkedAnimGraphBase>(*NodeIt))
+		{
+			return true;
+		}
+	}
+
+	return false;
+}
+
 void FAnimationBlueprintEditor::OnTogglePoseWatch()
 {
 	ReleaseAllManagedNodes();
@@ -853,6 +874,19 @@ void FAnimationBlueprintEditor::OnTogglePoseWatch()
 	AcquireAllManagedNodes();
 }
 
+void FAnimationBlueprintEditor::OnHideUnboundPropertyPins()
+{
+	const FGraphPanelSelectionSet SelectedNodes = GetSelectedNodes();
+
+	for (FGraphPanelSelectionSet::TConstIterator NodeIt(SelectedNodes); NodeIt; ++NodeIt)
+	{
+		if (UAnimGraphNode_LinkedAnimGraphBase* SelectedNode = Cast<UAnimGraphNode_LinkedAnimGraphBase>(*NodeIt))
+		{
+			UAnimationGraphSchema::HideUnboundPropertyPins(SelectedNode);
+		}
+	}
+}
+
 // Helper function for node conversions
 static void CopyPinData(UEdGraphNode* InOldNode, UEdGraphNode* InNewNode, const TCHAR* InPinName)
 {
diff --git a/Engine/Source/Editor/AnimationBlueprintEditor/Private/AnimationBlueprintEditor.h b/Engine/Source/Editor/AnimationBlueprintEditor/Private/AnimationBlueprintEditor.h
--- a/Engine/Source/Editor/AnimationBlueprintEditor/Private/AnimationBlueprintEditor.h
+++ b/Engine/Source/Editor/AnimationBlueprintEditor/Private/AnimationBlueprintEditor.h
@@ -236,6 +236,10 @@ protected:
 	bool CanTogglePoseWatch();
 	void OnTogglePoseWatch();
 
+	// Hide unbound pins on selected nodes 
+	bool CanHideUnboundPropertyPins();
+	void OnHideUnboundPropertyPins();
+
 	void BindCommands();
 
 protected:
