Fixed crash deleting a blendspace node

Guarded against a null graph selection in IsEditable()

diff --git a/Engine/Source/Editor/AnimationBlueprintEditor/Private/AnimationBlueprintEditor.cpp b/Engine/Source/Editor/AnimationBlueprintEditor/Private/AnimationBlueprintEditor.cpp
--- a/Engine/Source/Editor/AnimationBlueprintEditor/Private/AnimationBlueprintEditor.cpp
+++ b/Engine/Source/Editor/AnimationBlueprintEditor/Private/AnimationBlueprintEditor.cpp
@@ -1712,7 +1712,10 @@ bool FAnimationBlueprintEditor::IsEditable(UEdGraph* InGraph) const
 {
 	bool bEditable = FBlueprintEditor::IsEditable(InGraph);
 
-	bEditable &= (InGraph->GetTypedOuter<UBlueprint>() == GetBlueprintObj());
+	if (InGraph)
+	{
+		bEditable &= (InGraph->GetTypedOuter<UBlueprint>() == GetBlueprintObj());
+	}
 
 	return bEditable;
 }
