Truncation fixes for AnimationBlueprintEditor

diff --git a/Engine/Source/Editor/AnimationBlueprintEditor/Private/AnimGraphConnectionDrawingPolicy.cpp b/Engine/Source/Editor/AnimationBlueprintEditor/Private/AnimGraphConnectionDrawingPolicy.cpp
--- a/Engine/Source/Editor/AnimationBlueprintEditor/Private/AnimGraphConnectionDrawingPolicy.cpp
+++ b/Engine/Source/Editor/AnimationBlueprintEditor/Private/AnimGraphConnectionDrawingPolicy.cpp
@@ -140,7 +140,7 @@ void FAnimGraphConnectionDrawingPolicy::DetermineStyleOfExecWire(float& Thicknes
 	const float LightBlendThickness = SustainWireThickness;
 
 	Thickness = FMath::Lerp<float>(LightBlendThickness, HeavyBlendThickness, BlendWeight);
-	WireColor = WireColor * (BlendWeight * 0.5f + 0.5f);//FMath::Lerp<FLinearColor>(SustainColor, AttackColor, BlendWeight);
+	WireColor = WireColor * static_cast<float>(BlendWeight * 0.5 + 0.5);//FMath::Lerp<FLinearColor>(SustainColor, AttackColor, BlendWeight);
 
 	bDrawBubbles = true;
 }
diff --git a/Engine/Source/Editor/AnimationBlueprintEditor/Private/AnimationStateNodes/SGraphNodeAnimTransition.cpp b/Engine/Source/Editor/AnimationBlueprintEditor/Private/AnimationStateNodes/SGraphNodeAnimTransition.cpp
--- a/Engine/Source/Editor/AnimationBlueprintEditor/Private/AnimationStateNodes/SGraphNodeAnimTransition.cpp
+++ b/Engine/Source/Editor/AnimationBlueprintEditor/Private/AnimationStateNodes/SGraphNodeAnimTransition.cpp
@@ -442,8 +442,8 @@ void SGraphNodeAnimTransition::PositionBetweenTwoNodesWithOffset(const FGeometry
 	// Now we need to adjust the new center by the node size, zoom factor and multi node offset
 	const FVector2D NewCorner = NewCenter - (0.5f * DesiredNodeSize) + (DeltaNormal * MultiNodeOffset * DesiredNodeSize.Size());
 
-	GraphNode->NodePosX = NewCorner.X;
-	GraphNode->NodePosY = NewCorner.Y;
+	GraphNode->NodePosX = static_cast<int32>(NewCorner.X);
+	GraphNode->NodePosY = static_cast<int32>(NewCorner.Y);
 }
 
 /////////////////////////////////////////////////////
diff --git a/Engine/Source/Editor/AnimationBlueprintEditor/Private/StateMachineConnectionDrawingPolicy.cpp b/Engine/Source/Editor/AnimationBlueprintEditor/Private/StateMachineConnectionDrawingPolicy.cpp
--- a/Engine/Source/Editor/AnimationBlueprintEditor/Private/StateMachineConnectionDrawingPolicy.cpp
+++ b/Engine/Source/Editor/AnimationBlueprintEditor/Private/StateMachineConnectionDrawingPolicy.cpp
@@ -159,7 +159,7 @@ void FStateMachineConnectionDrawingPolicy::Internal_DrawLineWithArrow(const FVec
 
 	// Draw the arrow
 	const FVector2D ArrowDrawPos = EndPoint - ArrowRadius;
-	const float AngleInRadians = FMath::Atan2(DeltaPos.Y, DeltaPos.X);
+	const double AngleInRadians = FMath::Atan2(DeltaPos.Y, DeltaPos.X);
 
 	FSlateDrawElement::MakeRotatedBox(
 		DrawElementsList,
@@ -167,7 +167,7 @@ void FStateMachineConnectionDrawingPolicy::Internal_DrawLineWithArrow(const FVec
 		FPaintGeometry(ArrowDrawPos, ArrowImage->ImageSize * ZoomFactor, ZoomFactor),
 		ArrowImage,
 		ESlateDrawEffect::None,
-		AngleInRadians,
+		static_cast<float>(AngleInRadians),
 		TOptional<FVector2D>(),
 		FSlateDrawElement::RelativeToElement,
 		Params.WireColor
