[UE-169779] Transition rule relinking can be difficult to drag based on cursor placement

diff --git a/Engine/Source/Editor/AnimationBlueprintEditor/Private/StateMachineConnectionDrawingPolicy.cpp b/Engine/Source/Editor/AnimationBlueprintEditor/Private/StateMachineConnectionDrawingPolicy.cpp
--- a/Engine/Source/Editor/AnimationBlueprintEditor/Private/StateMachineConnectionDrawingPolicy.cpp
+++ b/Engine/Source/Editor/AnimationBlueprintEditor/Private/StateMachineConnectionDrawingPolicy.cpp
@@ -209,13 +209,9 @@ void FStateMachineConnectionDrawingPolicy::Internal_DrawLineWithArrow(const FVec
 		// Set the hovered pin results. This will be used by the SGraphPanel again.
 		const float SquaredDistToPin1 = (Params.AssociatedPin1 != nullptr) ? (StartPoint - LocalMousePosition).SizeSquared() : FLT_MAX;
 		const float SquaredDistToPin2 = (Params.AssociatedPin2 != nullptr) ? (EndPoint - LocalMousePosition).SizeSquared() : FLT_MAX;
-		if (StartHovered)
+		if (EndHovered)
 		{
-			SplineOverlapResult = FGraphSplineOverlapResult(Params.AssociatedPin2, Params.AssociatedPin1, SquaredDistToPin2, SquaredDistToPin2, SquaredDistToPin1, true);
-		}
-		else if (EndHovered)
-		{
-			SplineOverlapResult = FGraphSplineOverlapResult(Params.AssociatedPin1, Params.AssociatedPin2, SquaredDistToPin1, SquaredDistToPin1, SquaredDistToPin2, true);
+			SplineOverlapResult = FGraphSplineOverlapResult(Params.AssociatedPin1, Params.AssociatedPin2, SquaredDistToPin2, SquaredDistToPin1, SquaredDistToPin2, true);
 		}
 
 		// Draw grab handles only in case no relinking operation is performed
