Add missing anim BP template features now that skeleton compatibility has been relaxed

diff --git a/Engine/Source/Editor/AnimGraph/Private/AnimationGraphSchema.cpp b/Engine/Source/Editor/AnimGraph/Private/AnimationGraphSchema.cpp
--- a/Engine/Source/Editor/AnimGraph/Private/AnimationGraphSchema.cpp
+++ b/Engine/Source/Editor/AnimGraph/Private/AnimationGraphSchema.cpp
@@ -44,6 +44,29 @@
 
 #define LOCTEXT_NAMESPACE "AnimationGraphSchema"
 
+namespace UE::Anim::BP::Editor
+{
+	bool IsSkeletonCompatible(const UAnimBlueprint* AnimBlueprint, const UAnimationAsset* Asset)
+	{
+		if (AnimBlueprint == nullptr)
+		{
+			return false;	// No blueprint provided, cannot be compatible
+		}
+
+		if (AnimBlueprint->bIsTemplate)
+		{
+			return true;	// Templates are always compatible
+		}
+
+		if (AnimBlueprint->TargetSkeleton == nullptr)
+		{
+			return false;	// No target skeleton provided, cannot be compatible
+		}
+
+		return AnimBlueprint->TargetSkeleton->IsCompatibleForEditor(Asset->GetSkeleton());
+	}
+}
+
 /////////////////////////////////////////////////////
 // FAnimationLayerDragDropAction
 /** DragDropAction class for drag and dropping animation layers */
@@ -564,7 +587,7 @@ void UAnimationGraphSchema::SpawnNodeFromAsset(UAnimationAsset* Asset, const FVe
 
 	UAnimBlueprint* AnimBlueprint = Cast<UAnimBlueprint>(FBlueprintEditorUtils::FindBlueprintForGraph(Graph));
 
-	const bool bSkelMatch = (AnimBlueprint != nullptr) && (AnimBlueprint->TargetSkeleton != nullptr) && (AnimBlueprint->TargetSkeleton->IsCompatibleForEditor(Asset->GetSkeleton()));
+	const bool bSkelMatch = UE::Anim::BP::Editor::IsSkeletonCompatible(AnimBlueprint, Asset);
 	const bool bTypeMatch = (PinIfAvailable == nullptr) || UAnimationGraphSchema::IsLocalSpacePosePin(PinIfAvailable->PinType);
 	const bool bDirectionMatch = (PinIfAvailable == nullptr) || (PinIfAvailable->Direction == EGPD_Input);
 
@@ -673,7 +696,7 @@ void UAnimationGraphSchema::GetAssetsNodeHoverMessage(const TArray<FAssetData>&
 
 	// this one only should happen when there is an Anim Blueprint
 	UAnimBlueprint* AnimBlueprint = Cast<UAnimBlueprint>(FBlueprintEditorUtils::FindBlueprintForNode(HoverNode));
-	const bool bSkelMatch = (AnimBlueprint != NULL) && (AnimBlueprint->TargetSkeleton->IsCompatibleForEditor(Asset->GetSkeleton()));
+	const bool bSkelMatch = UE::Anim::BP::Editor::IsSkeletonCompatible(AnimBlueprint, Asset);
 
 	if (!bSkelMatch)
 	{
@@ -705,7 +728,7 @@ void UAnimationGraphSchema::GetAssetsPinHoverMessage(const TArray<FAssetData>& A
 	// this one only should happen when there is an Anim Blueprint
 	UAnimBlueprint* AnimBlueprint = Cast<UAnimBlueprint>(FBlueprintEditorUtils::FindBlueprintForNode(HoverPin->GetOwningNode()));
 
-	const bool bSkelMatch = (AnimBlueprint != NULL) && (AnimBlueprint->TargetSkeleton->IsCompatibleForEditor(Asset->GetSkeleton()));
+	const bool bSkelMatch = UE::Anim::BP::Editor::IsSkeletonCompatible(AnimBlueprint, Asset);
 	const bool bTypeMatch = UAnimationGraphSchema::IsLocalSpacePosePin(HoverPin->PinType);
 	const bool bDirectionMatch = HoverPin->Direction == EGPD_Input;
 
@@ -726,7 +749,7 @@ void UAnimationGraphSchema::GetAssetsGraphHoverMessage(const TArray<FAssetData>&
 	if (UAnimationAsset* AnimationAsset = FAssetData::GetFirstAsset<UAnimationAsset>(Assets))
 	{
 		UAnimBlueprint* AnimBlueprint = Cast<UAnimBlueprint>(FBlueprintEditorUtils::FindBlueprintForGraph(HoverGraph));
-		const bool bSkelMatch = (AnimBlueprint != nullptr) && (AnimBlueprint->TargetSkeleton != nullptr) && (AnimBlueprint->TargetSkeleton->IsCompatibleForEditor(AnimationAsset->GetSkeleton()));
+		const bool bSkelMatch = UE::Anim::BP::Editor::IsSkeletonCompatible(AnimBlueprint, AnimationAsset);
 		if (!bSkelMatch)
 		{
 			OutOkIcon = false;
diff --git a/Engine/Source/Editor/AnimationBlueprintEditor/Private/AnimationBlueprintEditor.cpp b/Engine/Source/Editor/AnimationBlueprintEditor/Private/AnimationBlueprintEditor.cpp
--- a/Engine/Source/Editor/AnimationBlueprintEditor/Private/AnimationBlueprintEditor.cpp
+++ b/Engine/Source/Editor/AnimationBlueprintEditor/Private/AnimationBlueprintEditor.cpp
@@ -44,7 +44,6 @@
 #include "AnimationBlueprintEditorModule.h"
 #include "AnimationBlueprintEditorSettings.h"
 #include "AnimationBlueprintInterfaceEditorMode.h"
-#include "AnimationBlueprintTemplateEditorMode.h"
 #include "AnimationEditorPreviewScene.h"
 #include "AnimationEditorUtils.h"
 #include "AnimationGraph.h"
@@ -462,6 +461,8 @@ void FAnimationBlueprintEditor::InitAnimationBlueprintEditor(const EToolkitMode:
 	// Register document editor for blendspaces
 	DocumentManager->RegisterDocumentFactory(MakeShared<FBlendSpaceDocumentTabFactory>(SharedThis(this)));
 
+	bool bHasBlueprintPreview = false;
+
 	if(InAnimBlueprint->BlueprintType == BPTYPE_Interface)
 	{
 		AddApplicationMode(
@@ -479,7 +480,9 @@ void FAnimationBlueprintEditor::InitAnimationBlueprintEditor(const EToolkitMode:
 	{
 		AddApplicationMode(
 			FAnimationBlueprintEditorModes::AnimationBlueprintTemplateEditorMode,
-			MakeShareable(new FAnimationBlueprintTemplateEditorMode(SharedThis(this))));
+			MakeShareable(new FAnimationBlueprintEditorMode(SharedThis(this))));
+
+		bHasBlueprintPreview = true;
 		
 		ExtendMenu();
 		ExtendToolbar();
@@ -494,10 +497,22 @@ void FAnimationBlueprintEditor::InitAnimationBlueprintEditor(const EToolkitMode:
 			FAnimationBlueprintEditorModes::AnimationBlueprintEditorMode,
 			MakeShareable(new FAnimationBlueprintEditorMode(SharedThis(this))));
 
+		bHasBlueprintPreview = true;
+
+		ExtendMenu();
+		ExtendToolbar();
+		RegenerateMenusAndToolbars();
+
+		// Activate the initial mode (which will populate with a real layout)
+		SetCurrentMode(FAnimationBlueprintEditorModes::AnimationBlueprintEditorMode);
+	}
+
+	if (bHasBlueprintPreview)
+	{
 		UDebugSkelMeshComponent* PreviewMeshComponent = PersonaToolkit->GetPreviewMeshComponent();
 		UAnimBlueprint* AnimBlueprint = PersonaToolkit->GetAnimBlueprint();
 		UAnimBlueprint* PreviewAnimBlueprint = AnimBlueprint->GetPreviewAnimationBlueprint();
-		
+
 		if (PreviewAnimBlueprint)
 		{
 			PersonaToolkit->GetPreviewScene()->SetPreviewAnimationBlueprint(PreviewAnimBlueprint, AnimBlueprint);
@@ -509,13 +524,6 @@ void FAnimationBlueprintEditor::InitAnimationBlueprintEditor(const EToolkitMode:
 		}
 
 		PersonaUtils::SetObjectBeingDebugged(AnimBlueprint, PreviewMeshComponent->GetAnimInstance());
-
-		ExtendMenu();
-		ExtendToolbar();
-		RegenerateMenusAndToolbars();
-
-		// Activate the initial mode (which will populate with a real layout)
-		SetCurrentMode(FAnimationBlueprintEditorModes::AnimationBlueprintEditorMode);
 	}
 
 	// Post-layout initialization
diff --git a/Engine/Source/Editor/AnimationBlueprintEditor/Private/AnimationBlueprintEditor.h b/Engine/Source/Editor/AnimationBlueprintEditor/Private/AnimationBlueprintEditor.h
--- a/Engine/Source/Editor/AnimationBlueprintEditor/Private/AnimationBlueprintEditor.h
+++ b/Engine/Source/Editor/AnimationBlueprintEditor/Private/AnimationBlueprintEditor.h
@@ -91,6 +91,7 @@ namespace AnimationBlueprintEditorTabs
 class FAnimationBlueprintEditor : public IAnimationBlueprintEditor
 {
 	friend class FAnimationBlueprintEditorMode;
+	friend class FAnimationBlueprintTemplateEditorMode;
 
 public:
 	/**
diff --git a/Engine/Source/Editor/AnimationBlueprintEditor/Private/AnimationBlueprintEditorMode.cpp b/Engine/Source/Editor/AnimationBlueprintEditor/Private/AnimationBlueprintEditorMode.cpp
--- a/Engine/Source/Editor/AnimationBlueprintEditor/Private/AnimationBlueprintEditorMode.cpp
+++ b/Engine/Source/Editor/AnimationBlueprintEditor/Private/AnimationBlueprintEditorMode.cpp
@@ -31,13 +31,55 @@ class UToolMenu;
 /////////////////////////////////////////////////////
 // FAnimationBlueprintEditorMode
 
+namespace UE::Anim::BP::Editor
+{
+	TSharedRef<FTabManager::FStack> CreateLeftBottomStack(bool bIsTemplate)
+	{
+		TSharedRef<FTabManager::FStack> Stack = FTabManager::NewStack();
+
+		if (!bIsTemplate)
+		{
+			Stack->AddTab(AnimationBlueprintEditorTabs::CurveNamesTab, ETabState::ClosedTab);
+			Stack->AddTab(AnimationBlueprintEditorTabs::SkeletonTreeTab, ETabState::ClosedTab);
+		}
+
+		Stack->AddTab(AnimationBlueprintEditorTabs::PoseWatchTab, ETabState::OpenedTab);
+		Stack->AddTab(FBlueprintEditorTabs::MyBlueprintID, ETabState::OpenedTab);
+
+		return Stack;
+	}
+
+	TSharedRef<FTabManager::FStack> CreateRightBottomStack(bool bIsTemplate)
+	{
+		TSharedRef<FTabManager::FStack> Stack = FTabManager::NewStack();
+
+		Stack->AddTab(AnimationBlueprintEditorTabs::AnimBlueprintPreviewEditorTab, ETabState::OpenedTab);
+		Stack->AddTab(AnimationBlueprintEditorTabs::AssetBrowserTab, ETabState::OpenedTab);
+
+		if (!bIsTemplate)
+		{
+			Stack->AddTab(AnimationBlueprintEditorTabs::SlotNamesTab, ETabState::ClosedTab);
+		}
+
+		Stack->SetForegroundTab(AnimationBlueprintEditorTabs::AnimBlueprintPreviewEditorTab);
+
+		return Stack;
+	}
+}
+
 FAnimationBlueprintEditorMode::FAnimationBlueprintEditorMode(const TSharedRef<FAnimationBlueprintEditor>& InAnimationBlueprintEditor)
 	: FBlueprintEditorApplicationMode(InAnimationBlueprintEditor, FAnimationBlueprintEditorModes::AnimationBlueprintEditorMode, FAnimationBlueprintEditorModes::GetLocalizedMode, false, false)
 {
 	PreviewScenePtr = InAnimationBlueprintEditor->GetPreviewScene();
 	AnimBlueprintPtr = CastChecked<UAnimBlueprint>(InAnimationBlueprintEditor->GetBlueprintObj());
 
-	TabLayout = FTabManager::NewLayout( "Stanalone_AnimationBlueprintEditMode_Layout_v1.4" )
+	bool bIsTemplate = false;
+	if (UAnimBlueprint* AnimBlueprint = AnimBlueprintPtr.Get())
+	{
+		bIsTemplate = AnimBlueprint->bIsTemplate;
+	}
+
+	TabLayout = FTabManager::NewLayout( "Stanalone_AnimationBlueprintEditMode_Layout_v1.5" )
 		->AddArea
 		(
 			FTabManager::NewPrimaryArea()
@@ -64,12 +106,8 @@ FAnimationBlueprintEditorMode::FAnimationBlueprintEditorMode(const TSharedRef<FA
 					->Split
 					(
 						//	Left bottom - preview settings
-						FTabManager::NewStack()
+						UE::Anim::BP::Editor::CreateLeftBottomStack(bIsTemplate)
 						->SetSizeCoefficient(0.5f)
-						->AddTab(AnimationBlueprintEditorTabs::CurveNamesTab, ETabState::ClosedTab)
-						->AddTab(AnimationBlueprintEditorTabs::SkeletonTreeTab, ETabState::ClosedTab)
-						->AddTab(AnimationBlueprintEditorTabs::PoseWatchTab, ETabState::OpenedTab)
-						->AddTab(FBlueprintEditorTabs::MyBlueprintID, ETabState::OpenedTab)
 					)
 				)
 				->Split
@@ -114,21 +152,14 @@ FAnimationBlueprintEditorMode::FAnimationBlueprintEditorMode(const TSharedRef<FA
 					->Split
 					(
 						// Right bottom - Asset browser & advanced preview settings
-						FTabManager::NewStack()
+						UE::Anim::BP::Editor::CreateRightBottomStack(bIsTemplate)
 						->SetHideTabWell(false)
 						->SetSizeCoefficient(0.5f)
-						->AddTab(AnimationBlueprintEditorTabs::AnimBlueprintPreviewEditorTab, ETabState::OpenedTab)
-						->AddTab(AnimationBlueprintEditorTabs::AssetBrowserTab, ETabState::OpenedTab)
-						->AddTab(AnimationBlueprintEditorTabs::SlotNamesTab, ETabState::ClosedTab)
-						->SetForegroundTab(AnimationBlueprintEditorTabs::AnimBlueprintPreviewEditorTab)
 					)
 				)
 			)
 		);
 
-	ISkeletonEditorModule& SkeletonEditorModule = FModuleManager::LoadModuleChecked<ISkeletonEditorModule>("SkeletonEditor");
-	TabFactories.RegisterFactory(SkeletonEditorModule.CreateSkeletonTreeTabFactory(InAnimationBlueprintEditor, InAnimationBlueprintEditor->GetSkeletonTree()));
-
 	FPersonaModule& PersonaModule = FModuleManager::LoadModuleChecked<FPersonaModule>("Persona");
 
 	FPersonaViewportArgs ViewportArgs(InAnimationBlueprintEditor->GetPersonaToolkit()->GetPreviewScene());
@@ -142,10 +173,16 @@ FAnimationBlueprintEditorMode::FAnimationBlueprintEditorMode(const TSharedRef<FA
 	TabFactories.RegisterFactory(PersonaModule.CreateAnimationAssetBrowserTabFactory(InAnimationBlueprintEditor, InAnimationBlueprintEditor->GetPersonaToolkit(), FOnOpenNewAsset::CreateSP(&InAnimationBlueprintEditor.Get(), &FAnimationBlueprintEditor::HandleOpenNewAsset), FOnAnimationSequenceBrowserCreated(), true));
 	TabFactories.RegisterFactory(PersonaModule.CreateAnimBlueprintPreviewTabFactory(InAnimationBlueprintEditor, InAnimationBlueprintEditor->GetPersonaToolkit()->GetPreviewScene()));
 	TabFactories.RegisterFactory(PersonaModule.CreateAnimBlueprintAssetOverridesTabFactory(InAnimationBlueprintEditor, InAnimationBlueprintEditor->GetPersonaToolkit()->GetAnimBlueprint(), InAnimationBlueprintEditor->OnPostUndo));
-	TabFactories.RegisterFactory(PersonaModule.CreateSkeletonSlotNamesTabFactory(InAnimationBlueprintEditor, InAnimationBlueprintEditor->GetSkeletonTree()->GetEditableSkeleton(), FOnObjectSelected::CreateSP(&InAnimationBlueprintEditor.Get(), &FAnimationBlueprintEditor::HandleObjectSelected)));
-	TabFactories.RegisterFactory(PersonaModule.CreateCurveViewerTabFactory(InAnimationBlueprintEditor, InAnimationBlueprintEditor->GetSkeletonTree()->GetEditableSkeleton(), InAnimationBlueprintEditor->GetPersonaToolkit()->GetPreviewScene(), FOnObjectsSelected::CreateSP(&InAnimationBlueprintEditor.Get(), &FAnimationBlueprintEditor::HandleObjectsSelected)));
 	TabFactories.RegisterFactory(PersonaModule.CreatePoseWatchTabFactory(InAnimationBlueprintEditor));
 
+	if (!bIsTemplate)
+	{
+		ISkeletonEditorModule& SkeletonEditorModule = FModuleManager::LoadModuleChecked<ISkeletonEditorModule>("SkeletonEditor");
+		TabFactories.RegisterFactory(SkeletonEditorModule.CreateSkeletonTreeTabFactory(InAnimationBlueprintEditor, InAnimationBlueprintEditor->GetSkeletonTree()));
+		TabFactories.RegisterFactory(PersonaModule.CreateSkeletonSlotNamesTabFactory(InAnimationBlueprintEditor, InAnimationBlueprintEditor->GetSkeletonTree()->GetEditableSkeleton(), FOnObjectSelected::CreateSP(&InAnimationBlueprintEditor.Get(), &FAnimationBlueprintEditor::HandleObjectSelected)));
+		TabFactories.RegisterFactory(PersonaModule.CreateCurveViewerTabFactory(InAnimationBlueprintEditor, InAnimationBlueprintEditor->GetSkeletonTree()->GetEditableSkeleton(), InAnimationBlueprintEditor->GetPersonaToolkit()->GetPreviewScene(), FOnObjectsSelected::CreateSP(&InAnimationBlueprintEditor.Get(), &FAnimationBlueprintEditor::HandleObjectsSelected)));
+	}
+
 	// setup toolbar - clear existing toolbar extender from the BP mode
 	//@TODO: Keep this in sync with BlueprintEditorModes.cpp
 	ToolbarExtender = MakeShareable(new FExtender);
diff --git a/Engine/Source/Editor/AnimationBlueprintEditor/Private/AnimationBlueprintTemplateEditorMode.cpp b/Engine/Source/Editor/AnimationBlueprintEditor/Private/AnimationBlueprintTemplateEditorMode.cpp
deleted file mode 100644
diff --git a/Engine/Source/Editor/AnimationBlueprintEditor/Private/AnimationBlueprintTemplateEditorMode.h b/Engine/Source/Editor/AnimationBlueprintEditor/Private/AnimationBlueprintTemplateEditorMode.h
deleted file mode 100644
diff --git a/Engine/Source/Editor/Persona/Private/AnimGraphNodeDetails.cpp b/Engine/Source/Editor/Persona/Private/AnimGraphNodeDetails.cpp
--- a/Engine/Source/Editor/Persona/Private/AnimGraphNodeDetails.cpp
+++ b/Engine/Source/Editor/Persona/Private/AnimGraphNodeDetails.cpp
@@ -98,6 +98,7 @@ void FAnimGraphNodeDetails::CustomizeDetails(class IDetailLayoutBuilder& DetailB
 
 	TargetSkeleton = AnimGraphNode->HasValidBlueprint() ? AnimGraphNode->GetAnimBlueprint()->TargetSkeleton : nullptr;
 	TargetSkeletonName = TargetSkeleton ? FObjectPropertyBase::GetExportPath(TargetSkeleton) : FString(TEXT(""));
+	bIsAnimBPTemplate = AnimGraphNode->HasValidBlueprint() ? AnimGraphNode->GetAnimBlueprint()->bIsTemplate : false;
 
 	// Get the node property
 	const FStructProperty* NodeProperty = AnimGraphNode->GetFNodeProperty();
@@ -275,13 +276,27 @@ TSharedRef<SWidget> FAnimGraphNodeDetails::CreatePropertyWidget(FProperty* Targe
 bool FAnimGraphNodeDetails::OnShouldFilterAnimAsset( const FAssetData& AssetData, UClass* NodeToFilterFor ) const
 {
 	FAssetDataTagMapSharedView::FFindTagResult Result = AssetData.TagsAndValues.FindTag("Skeleton");
-	if (Result.IsSet() && TargetSkeleton && TargetSkeleton->IsCompatibleForEditor(AssetData))
+	if (Result.IsSet())
 	{
-		const UClass* AssetClass = AssetData.GetClass();
-		// If node is an 'asset player', only let you select the right kind of asset for it
-		if (!NodeToFilterFor->IsChildOf(UAnimGraphNode_AssetPlayerBase::StaticClass()) || (AssetClass && SupportNodeClassForAsset(AssetClass, NodeToFilterFor)))
+		bool bIsAssetCompatible;
+		if (bIsAnimBPTemplate)
 		{
-			return false;
+			// If we are a template, we are always compatible
+			bIsAssetCompatible = true;
+		}
+		else
+		{
+			bIsAssetCompatible = TargetSkeleton && TargetSkeleton->IsCompatibleForEditor(AssetData);
+		}
+
+		if (bIsAssetCompatible)
+		{
+			const UClass* AssetClass = AssetData.GetClass();
+			// If node is an 'asset player', only let you select the right kind of asset for it
+			if (!NodeToFilterFor->IsChildOf(UAnimGraphNode_AssetPlayerBase::StaticClass()) || (AssetClass && SupportNodeClassForAsset(AssetClass, NodeToFilterFor)))
+			{
+				return false;
+			}
 		}
 	}
 	return true;
diff --git a/Engine/Source/Editor/Persona/Private/AnimGraphNodeDetails.h b/Engine/Source/Editor/Persona/Private/AnimGraphNodeDetails.h
--- a/Engine/Source/Editor/Persona/Private/AnimGraphNodeDetails.h
+++ b/Engine/Source/Editor/Persona/Private/AnimGraphNodeDetails.h
@@ -74,6 +74,9 @@ protected:
 
 	/** Builder that built this customization (used for refresh) */
 	IDetailLayoutBuilder* DetailLayoutBuilder = nullptr;
+
+	/** Whether or not the anim blueprint is a template */
+	bool bIsAnimBPTemplate = false;
 };
 
 /////////////////////////////////////////////////////
diff --git a/Engine/Source/Editor/Persona/Private/AnimationEditorPreviewScene.cpp b/Engine/Source/Editor/Persona/Private/AnimationEditorPreviewScene.cpp
--- a/Engine/Source/Editor/Persona/Private/AnimationEditorPreviewScene.cpp
+++ b/Engine/Source/Editor/Persona/Private/AnimationEditorPreviewScene.cpp
@@ -305,6 +305,12 @@ void FAnimationEditorPreviewScene::SetPreviewMeshInternal(USkeletalMesh* NewPrev
 		PersonaUtils::SetObjectBeingDebugged(SourceBlueprint, DebuggedSkeletalMeshComponent->GetAnimInstance());
 	}
 
+	// If we didn't have a preview mesh before and we select one now, set it up as the object being debugged
+	if (DebuggedSkeletalMeshComponent == nullptr && NewPreviewMesh != nullptr && SkeletalMeshComponent->GetAnimInstance() && SkeletalMeshComponent->GetAnimInstance()->IsA(SourceBlueprint->GeneratedClass))
+	{
+		PersonaUtils::SetObjectBeingDebugged(SourceBlueprint, SkeletalMeshComponent->GetAnimInstance());
+	}
+
 	OnPreviewMeshChanged.Broadcast(OldPreviewMesh, NewPreviewMesh);
 }
 
diff --git a/Engine/Source/Editor/Persona/Private/PreviewSceneCustomizations.cpp b/Engine/Source/Editor/Persona/Private/PreviewSceneCustomizations.cpp
--- a/Engine/Source/Editor/Persona/Private/PreviewSceneCustomizations.cpp
+++ b/Engine/Source/Editor/Persona/Private/PreviewSceneCustomizations.cpp
@@ -192,7 +192,8 @@ void FPreviewSceneDescriptionCustomization::CustomizeDetails(IDetailLayoutBuilde
 
 		const bool bCanUseDifferentSkeleton =
 			(PersonaToolkit.Pin()->GetContext() == UPhysicsAsset::StaticClass()->GetFName()) ||
-			(PersonaToolkit.Pin()->GetContext() == TEXT("ControlRigBlueprint"));
+			(PersonaToolkit.Pin()->GetContext() == TEXT("ControlRigBlueprint")) ||
+			(PersonaToolkit.Pin()->GetAnimBlueprint() != nullptr && PersonaToolkit.Pin()->GetAnimBlueprint()->bIsTemplate);
 
 		DetailBuilder.EditCategory("Mesh")
 		.AddProperty(SkeletalMeshProperty)
diff --git a/Engine/Source/Runtime/Engine/Private/AnimBlueprint.cpp b/Engine/Source/Runtime/Engine/Private/AnimBlueprint.cpp
--- a/Engine/Source/Runtime/Engine/Private/AnimBlueprint.cpp
+++ b/Engine/Source/Runtime/Engine/Private/AnimBlueprint.cpp
@@ -279,15 +279,11 @@ void UAnimBlueprint::PostEditChangeProperty(struct FPropertyChangedEvent& Proper
 USkeletalMesh* UAnimBlueprint::GetPreviewMesh(bool bFindIfNotSet/*=false*/)
 {
 #if WITH_EDITORONLY_DATA
-	USkeletalMesh* PreviewMesh = PreviewSkeletalMesh.LoadSynchronous();
-	// if somehow skeleton changes, just nullify it. 
-	if (PreviewMesh && PreviewMesh->GetSkeleton() != TargetSkeleton)
+	if (!PreviewSkeletalMesh.IsValid())
 	{
-		PreviewMesh = nullptr;
-		SetPreviewMesh(nullptr);
+		PreviewSkeletalMesh.LoadSynchronous();
 	}
-
-	return PreviewMesh;
+	return PreviewSkeletalMesh.Get();
 #else
 	return nullptr;
 #endif
diff --git a/Engine/Source/Runtime/Engine/Private/Components/SkeletalMeshComponent.cpp b/Engine/Source/Runtime/Engine/Private/Components/SkeletalMeshComponent.cpp
--- a/Engine/Source/Runtime/Engine/Private/Components/SkeletalMeshComponent.cpp
+++ b/Engine/Source/Runtime/Engine/Private/Components/SkeletalMeshComponent.cpp
@@ -3863,12 +3863,7 @@ void USkeletalMeshComponent::ValidateAnimation()
 		IAnimClassInterface* AnimClassInterface = IAnimClassInterface::GetFromClass(AnimClass);
 		if (AnimClassInterface && SkelMesh)
 		{
-			if (AnimClassInterface->GetTargetSkeleton() == nullptr)
-			{
-				UE_LOG(LogAnimation, Warning, TEXT("AnimBP %s is incompatible because it has no target skeleton, removing AnimBP from actor."), *AnimClass->GetName());
-				AnimClass = nullptr;
-			}
-			else if (SkelMesh->GetSkeleton() == nullptr)
+			if (SkelMesh->GetSkeleton() == nullptr)
 			{
 				UE_LOG(LogAnimation, Warning, TEXT("AnimBP %s is incompatible because mesh %s has no skeleton, removing AnimBP from actor."), *AnimClass->GetName(), *SkelMesh->GetName());
 				AnimClass = nullptr;
