New colors to distinguish state aliases and conduits.

diff --git a/Engine/Source/Editor/AnimationBlueprintEditor/Private/AnimationStateNodes/SGraphNodeAnimState.cpp b/Engine/Source/Editor/AnimationBlueprintEditor/Private/AnimationStateNodes/SGraphNodeAnimState.cpp
--- a/Engine/Source/Editor/AnimationBlueprintEditor/Private/AnimationStateNodes/SGraphNodeAnimState.cpp
+++ b/Engine/Source/Editor/AnimationBlueprintEditor/Private/AnimationStateNodes/SGraphNodeAnimState.cpp
@@ -420,6 +420,14 @@ void SGraphNodeAnimConduit::GetNodeInfoPopups(FNodeInfoContext* Context, TArray<
 	// Intentionally empty.
 }
 
+FSlateColor SGraphNodeAnimConduit::GetBorderBackgroundColor_Internal(FLinearColor InactiveStateColor, FLinearColor ActiveStateColorDim, FLinearColor ActiveStateColorBright) const
+{
+	// Override inactive state color for conduits.
+	InactiveStateColor = FLinearColor(0.38f, 0.45f, 0.21f);
+
+	return SGraphNodeAnimState::GetBorderBackgroundColor_Internal(InactiveStateColor, ActiveStateColorDim, ActiveStateColorBright);
+}
+
 FText SGraphNodeAnimConduit::GetPreviewCornerText() const
 {
 	UAnimStateNodeBase* StateNode = CastChecked<UAnimStateNodeBase>(GraphNode);
diff --git a/Engine/Source/Editor/AnimationBlueprintEditor/Private/AnimationStateNodes/SGraphNodeAnimState.h b/Engine/Source/Editor/AnimationBlueprintEditor/Private/AnimationStateNodes/SGraphNodeAnimState.h
--- a/Engine/Source/Editor/AnimationBlueprintEditor/Private/AnimationStateNodes/SGraphNodeAnimState.h
+++ b/Engine/Source/Editor/AnimationBlueprintEditor/Private/AnimationStateNodes/SGraphNodeAnimState.h
@@ -67,6 +67,8 @@ public:
 	virtual void GetNodeInfoPopups(FNodeInfoContext* Context, TArray<FGraphInformationPopupInfo>& Popups) const override;
 	// End of SNodePanel::SNode interface
 protected:
+	virtual FSlateColor GetBorderBackgroundColor_Internal(FLinearColor InactiveStateColor, FLinearColor ActiveStateColorDim, FLinearColor ActiveStateColorBright) const override;
+
 	virtual FText GetPreviewCornerText() const override;
 	virtual const FSlateBrush* GetNameIcon() const override;
 };
diff --git a/Engine/Source/Editor/AnimationBlueprintEditor/Private/AnimationStateNodes/SGraphNodeAnimStateAlias.cpp b/Engine/Source/Editor/AnimationBlueprintEditor/Private/AnimationStateNodes/SGraphNodeAnimStateAlias.cpp
--- a/Engine/Source/Editor/AnimationBlueprintEditor/Private/AnimationStateNodes/SGraphNodeAnimStateAlias.cpp
+++ b/Engine/Source/Editor/AnimationBlueprintEditor/Private/AnimationStateNodes/SGraphNodeAnimStateAlias.cpp
@@ -46,6 +46,9 @@ void SGraphNodeAnimStateAlias::GetNodeInfoPopups(FNodeInfoContext* Context, TArr
 
 FSlateColor SGraphNodeAnimStateAlias::GetBorderBackgroundColor_Internal(FLinearColor InactiveStateColor, FLinearColor ActiveStateColorDim, FLinearColor ActiveStateColorBright) const
 {
+	// Override inactive state color for aliases.
+	InactiveStateColor = FLinearColor(0.45f, 0.33f, 0.37f);
+
 	UAnimBlueprint* AnimBlueprint = Cast<UAnimBlueprint>(FBlueprintEditorUtils::FindBlueprintForNode(GraphNode));
 	if (AnimBlueprint)
 	{
