Fix for static analysis warning: "Parameter 'InactiveStateColor' is always rewritten in function body before being used"

diff --git a/Engine/Source/Editor/AnimationBlueprintEditor/Private/AnimationStateNodes/SGraphNodeAnimState.cpp b/Engine/Source/Editor/AnimationBlueprintEditor/Private/AnimationStateNodes/SGraphNodeAnimState.cpp
--- a/Engine/Source/Editor/AnimationBlueprintEditor/Private/AnimationStateNodes/SGraphNodeAnimState.cpp
+++ b/Engine/Source/Editor/AnimationBlueprintEditor/Private/AnimationStateNodes/SGraphNodeAnimState.cpp
@@ -423,9 +423,7 @@ void SGraphNodeAnimConduit::GetNodeInfoPopups(FNodeInfoContext* Context, TArray<
 FSlateColor SGraphNodeAnimConduit::GetBorderBackgroundColor_Internal(FLinearColor InactiveStateColor, FLinearColor ActiveStateColorDim, FLinearColor ActiveStateColorBright) const
 {
 	// Override inactive state color for conduits.
-	InactiveStateColor = FLinearColor(0.38f, 0.45f, 0.21f);
-
-	return SGraphNodeAnimState::GetBorderBackgroundColor_Internal(InactiveStateColor, ActiveStateColorDim, ActiveStateColorBright);
+	return SGraphNodeAnimState::GetBorderBackgroundColor_Internal(FLinearColor(0.38f, 0.45f, 0.21f), ActiveStateColorDim, ActiveStateColorBright);
 }
 
 FText SGraphNodeAnimConduit::GetPreviewCornerText() const
diff --git a/Engine/Source/Editor/AnimationBlueprintEditor/Private/AnimationStateNodes/SGraphNodeAnimStateAlias.cpp b/Engine/Source/Editor/AnimationBlueprintEditor/Private/AnimationStateNodes/SGraphNodeAnimStateAlias.cpp
--- a/Engine/Source/Editor/AnimationBlueprintEditor/Private/AnimationStateNodes/SGraphNodeAnimStateAlias.cpp
+++ b/Engine/Source/Editor/AnimationBlueprintEditor/Private/AnimationStateNodes/SGraphNodeAnimStateAlias.cpp
@@ -46,9 +46,6 @@ void SGraphNodeAnimStateAlias::GetNodeInfoPopups(FNodeInfoContext* Context, TArr
 
 FSlateColor SGraphNodeAnimStateAlias::GetBorderBackgroundColor_Internal(FLinearColor InactiveStateColor, FLinearColor ActiveStateColorDim, FLinearColor ActiveStateColorBright) const
 {
-	// Override inactive state color for aliases.
-	InactiveStateColor = FLinearColor(0.45f, 0.33f, 0.37f);
-
 	UAnimBlueprint* AnimBlueprint = Cast<UAnimBlueprint>(FBlueprintEditorUtils::FindBlueprintForNode(GraphNode));
 	if (AnimBlueprint)
 	{
@@ -106,7 +103,8 @@ FSlateColor SGraphNodeAnimStateAlias::GetBorderBackgroundColor_Internal(FLinearC
 		}
 	}
 
-	return InactiveStateColor;
+	// Override inactive state color for aliases.
+	return FLinearColor(0.45f, 0.33f, 0.37f);
 }
 
 FText SGraphNodeAnimStateAlias::GetPreviewCornerText() const
