Fix preview instance not being re-selecteable after selecting an external anim instance to debug

diff --git a/Engine/Source/Editor/AnimationBlueprintEditor/Private/AnimationBlueprintEditor.cpp b/Engine/Source/Editor/AnimationBlueprintEditor/Private/AnimationBlueprintEditor.cpp
--- a/Engine/Source/Editor/AnimationBlueprintEditor/Private/AnimationBlueprintEditor.cpp
+++ b/Engine/Source/Editor/AnimationBlueprintEditor/Private/AnimationBlueprintEditor.cpp
@@ -1669,6 +1669,12 @@ UAnimInstance* FAnimationBlueprintEditor::GetPreviewInstance() const
 void FAnimationBlueprintEditor::GetCustomDebugObjects(TArray<FCustomDebugObject>& DebugList) const
 {
 	UAnimInstance* PreviewInstance = GetPreviewInstance();
+	if (PreviewInstance == nullptr)
+	{
+		UDebugSkelMeshComponent* PreviewMeshComponent = PersonaToolkit->GetPreviewMeshComponent();
+		PreviewInstance = PreviewMeshComponent->SavedAnimScriptInstance;
+	}
+
 	if (PreviewInstance)
 	{
 		new (DebugList) FCustomDebugObject(PreviewInstance, LOCTEXT("PreviewObjectLabel", "Preview Instance").ToString());
@@ -1681,6 +1687,12 @@ void FAnimationBlueprintEditor::GetCustomDebugObjects(TArray<FCustomDebugObject>
 FString FAnimationBlueprintEditor::GetCustomDebugObjectLabel(UObject* ObjectBeingDebugged) const
 {
 	UAnimInstance* PreviewInstance = GetPreviewInstance();
+	if (PreviewInstance == nullptr)
+	{
+		UDebugSkelMeshComponent* PreviewMeshComponent = PersonaToolkit->GetPreviewMeshComponent();
+		PreviewInstance = PreviewMeshComponent->SavedAnimScriptInstance;
+	}
+
 	if (PreviewInstance == ObjectBeingDebugged)
 	{
 		return LOCTEXT("PreviewObjectLabel", "Preview Instance").ToString();
