Fix for editor crash when closing one viewport while two viewports are open in persona

diff --git a/Engine/Source/Editor/AnimationBlueprintEditor/Private/AnimationBlueprintEditor.cpp b/Engine/Source/Editor/AnimationBlueprintEditor/Private/AnimationBlueprintEditor.cpp
--- a/Engine/Source/Editor/AnimationBlueprintEditor/Private/AnimationBlueprintEditor.cpp
+++ b/Engine/Source/Editor/AnimationBlueprintEditor/Private/AnimationBlueprintEditor.cpp
@@ -329,6 +329,16 @@ FAnimationBlueprintEditor::~FAnimationBlueprintEditor()
 	// NOTE: Any tabs that we still have hanging out when destroyed will be cleaned up by FBaseToolkit's destructor
 
 	SaveEditorSettings();
+
+	// Explicit Reset of the PersonaToolKit to force destruction of the PreviewScene
+	// This will call PreviewWorld->CleanupWorld() while the EditorModeManager still has the PreviewScene
+	PersonaToolkit.Reset();
+
+	// Now we have to clean the PersonaToolkit PreviewScene from the EditorModeManager, as it has been destroyed.
+	// This avoids a memory after delete use when any additional PreviewScene calls PreviewWorld->CleanupWorld(),
+	// as that function executes a callback that the EditorModeManager is registered to.
+	FEditorModeTools& ModeTools = GetEditorModeManager();
+	((FAssetEditorModeManager&)ModeTools).SetPreviewScene(nullptr);
 }
 
 void FAnimationBlueprintEditor::HandleUpdateSettings(const UAnimationBlueprintEditorSettings* AnimationBlueprintEditorSettings, EPropertyChangeType::Type ChangeType)
diff --git a/Engine/Source/Editor/Persona/Private/AnimationEditorViewportClient.cpp b/Engine/Source/Editor/Persona/Private/AnimationEditorViewportClient.cpp
--- a/Engine/Source/Editor/Persona/Private/AnimationEditorViewportClient.cpp
+++ b/Engine/Source/Editor/Persona/Private/AnimationEditorViewportClient.cpp
@@ -152,11 +152,6 @@ FAnimationViewportClient::~FAnimationViewportClient()
 {
 	CameraController = CachedDefaultCameraController;
 
-	if (ModeTools.IsValid())
-	{
-		((FAssetEditorModeManager*)ModeTools.Get())->SetPreviewScene(nullptr);
-	}
-
 	// Unregistering the callbacks is mandatory, else we get random crashes
 	if (FAnimationEditorPreviewScene* AnimationEditorPreviewScene = GetAnimPreviewScenePtr().Get())
 	{
