Improve infinite loop handling in animation blueprints

Change early-out for exception handling to skip exceptions for all object types in preview worlds, not just actors. Also handle excpetions in this case when the debugger is attached to the object already (allows for break on exceptions to work).
Hook the BP exception delegate to disable ticking on the skeletal mesh component so we dont keep ticking runaway scripts
Tweaked thumbnail request to not submit the same asset multiple times - this happens multiple times on compilation and can cause longer hitches when the thumbnail is requested and ticked with an infinite loop in place.

diff --git a/Engine/Source/Editor/AnimationBlueprintEditor/Private/AnimationBlueprintEditor.cpp b/Engine/Source/Editor/AnimationBlueprintEditor/Private/AnimationBlueprintEditor.cpp
--- a/Engine/Source/Editor/AnimationBlueprintEditor/Private/AnimationBlueprintEditor.cpp
+++ b/Engine/Source/Editor/AnimationBlueprintEditor/Private/AnimationBlueprintEditor.cpp
@@ -339,6 +339,8 @@ FAnimationBlueprintEditor::~FAnimationBlueprintEditor()
 	// as that function executes a callback that the EditorModeManager is registered to.
 	FEditorModeTools& ModeTools = GetEditorModeManager();
 	((FAssetEditorModeManager&)ModeTools).SetPreviewScene(nullptr);
+
+	FBlueprintCoreDelegates::OnScriptException.Remove(ScriptExceptionHandle);
 }
 
 void FAnimationBlueprintEditor::HandleUpdateSettings(const UAnimationBlueprintEditorSettings* AnimationBlueprintEditorSettings, EPropertyChangeType::Type ChangeType)
@@ -554,6 +556,8 @@ void FAnimationBlueprintEditor::InitAnimationBlueprintEditor(const EToolkitMode:
 		UAnimationBlueprintEditorSettings::FOnUpdateSettingsMulticaster::FDelegate::CreateSP(this, &FAnimationBlueprintEditor::HandleUpdateSettings));
 
 	PersonaToolkit->GetPreviewScene()->SetAllowMeshHitProxies(false);
+
+	ScriptExceptionHandle = FBlueprintCoreDelegates::OnScriptException.AddSP(this, &FAnimationBlueprintEditor::HandleScriptException);
 }
 
 void FAnimationBlueprintEditor::BindCommands()
@@ -1913,6 +1917,12 @@ void FAnimationBlueprintEditor::OnBlueprintPostCompile(UBlueprint* InBlueprint)
 		// We dont cache this persistently, only during a pre/post compile bracket
 		DebuggedMeshComponent = nullptr;
 	}
+
+	// Make sure we re-enable ticking when we recompile (e.g. to fix an infinite loop that paused us)
+	if (USkeletalMeshComponent* SkeletalMeshComponent = GetPreviewScene()->GetPreviewMeshComponent())
+	{
+		SkeletalMeshComponent->SetComponentTickEnabled(true);
+	}
 }
 
 void FAnimationBlueprintEditor::OnBlueprintChangedImpl(UBlueprint* InBlueprint, bool bIsJustBeingCompiled /*= false*/)
@@ -2572,6 +2582,47 @@ void FAnimationBlueprintEditor::HandleViewportCreated(const TSharedRef<IPersonaV
 		],
 		FPersonaViewportNotificationOptions(TAttribute<EVisibility>::Create(GetCompilationStateVisibility))
 	);
+
+	auto GetInfiniteLoopVisibility = [this]()
+	{
+		if (USkeletalMeshComponent* SkeletalMeshComponent = GetPreviewScene()->GetPreviewMeshComponent())
+		{
+			return SkeletalMeshComponent->IsComponentTickEnabled() ? EVisibility::Collapsed : EVisibility::Visible;
+		}
+		return EVisibility::Collapsed;
+	};
+
+	InPersonaViewport->AddNotification(EMessageSeverity::Error,
+		false,
+		SNew(SHorizontalBox)
+		.Visibility_Lambda(GetInfiniteLoopVisibility)
+		+SHorizontalBox::Slot()
+		.FillWidth(1.0f)
+		.Padding(4.0f, 4.0f)
+		[
+			SNew(SHorizontalBox)
+			.ToolTipText_Lambda(GetCompilationStateText)
+			+SHorizontalBox::Slot()
+			.AutoWidth()
+			.VAlign(VAlign_Center)
+			.Padding(0.0f, 0.0f, 4.0f, 0.0f)
+			[
+				SNew(STextBlock)
+				.TextStyle(FAppStyle::Get(), "AnimViewport.MessageText")
+				.Font(FAppStyle::Get().GetFontStyle("FontAwesome.9"))
+				.Text(FEditorFontGlyphs::Exclamation_Triangle)
+			]
+			+SHorizontalBox::Slot()
+			.VAlign(VAlign_Center)
+			.FillWidth(1.0f)
+			[
+				SNew(STextBlock)
+				.Text(LOCTEXT("InfiniteLoopDetected", "Infinite Loop Detected"))
+				.TextStyle(FAppStyle::Get(), "AnimViewport.MessageText")
+			]
+		],
+		FPersonaViewportNotificationOptions(TAttribute<EVisibility>::Create(GetInfiniteLoopVisibility))
+	);
 }
 
 void FAnimationBlueprintEditor::LoadEditorSettings()
@@ -2608,5 +2659,25 @@ void FAnimationBlueprintEditor::HandleAnimationSequenceBrowserCreated(const TSha
 	SequenceBrowser = InSequenceBrowser;
 }
 
+void FAnimationBlueprintEditor::HandleScriptException(const UObject* InObject, const FFrame& InFrame, const FBlueprintExceptionInfo& InInfo)
+{
+	// If the object is an anim instance in our preview world and is infinitely looping, disable ticking (renabled on recompilation)
+	if (InInfo.GetType() == EBlueprintExceptionType::InfiniteLoop)
+	{
+		if (InObject && InObject->IsA<UAnimInstance>())
+		{
+			UWorld* ObjectWorld = InObject->GetWorld();
+			TSharedRef<IPersonaPreviewScene> ThisPreviewScene = GetPreviewScene();
+			if (ObjectWorld == ThisPreviewScene->GetWorld())
+			{
+				if (USkeletalMeshComponent* SkeletalMeshComponent = ThisPreviewScene->GetPreviewMeshComponent())
+				{
+					SkeletalMeshComponent->SetComponentTickEnabled(false);
+				}
+			}
+		}
+	}
+}
+
 #undef LOCTEXT_NAMESPACE
 
diff --git a/Engine/Source/Editor/AnimationBlueprintEditor/Private/AnimationBlueprintEditor.h b/Engine/Source/Editor/AnimationBlueprintEditor/Private/AnimationBlueprintEditor.h
--- a/Engine/Source/Editor/AnimationBlueprintEditor/Private/AnimationBlueprintEditor.h
+++ b/Engine/Source/Editor/AnimationBlueprintEditor/Private/AnimationBlueprintEditor.h
@@ -43,6 +43,8 @@ class UObject;
 class USkeletalMesh;
 class USkeletalMeshComponent;
 struct FToolMenuContext;
+struct FFrame;
+struct FBlueprintExceptionInfo;
 
 struct FAnimationBlueprintEditorModes
 {
@@ -358,6 +360,9 @@ private:
 
 	void HandleAnimationSequenceBrowserCreated(const TSharedRef<IAnimationSequenceBrowser>& InSequenceBrowser);
 
+	/** Hook the BP exception handler to deal with infinite loops (more) gracefully */
+	void HandleScriptException(const UObject* InObject, const FFrame& InFrame, const FBlueprintExceptionInfo& InInfo);
+
 	void HandleUpdateSettings(const UAnimationBlueprintEditorSettings* AnimationBlueprintEditorSettings, EPropertyChangeType::Type ChangeType);
 
 	/** Chooses a suitable pose watch color automatically - i.e. one that isn't already in use (if possible) */
@@ -422,4 +427,7 @@ private:
 
 	/** Delegate handle registered for when settings change */
 	FDelegateHandle AnimationBlueprintEditorSettingsChangedHandle;
+
+	/** Delegate handle registered to handle infinite loop exceptions */
+	FDelegateHandle ScriptExceptionHandle;
 };
diff --git a/Engine/Source/Editor/UnrealEd/Private/AssetThumbnail.cpp b/Engine/Source/Editor/UnrealEd/Private/AssetThumbnail.cpp
--- a/Engine/Source/Editor/UnrealEd/Private/AssetThumbnail.cpp
+++ b/Engine/Source/Editor/UnrealEd/Private/AssetThumbnail.cpp
@@ -1626,7 +1626,7 @@ void FAssetThumbnailPool::RefreshThumbnailsFor( const FSoftObjectPath& ObjectPat
 	{
 		if ( ThumbIt.Key().ObjectPath == ObjectPath )
 		{
-			ThumbnailsToRenderStack.Push( ThumbIt.Value() );
+			ThumbnailsToRenderStack.AddUnique( ThumbIt.Value() );
 		}
 	}
 }
diff --git a/Engine/Source/Editor/UnrealEd/Private/Kismet2/KismetDebugUtilities.cpp b/Engine/Source/Editor/UnrealEd/Private/Kismet2/KismetDebugUtilities.cpp
--- a/Engine/Source/Editor/UnrealEd/Private/Kismet2/KismetDebugUtilities.cpp
+++ b/Engine/Source/Editor/UnrealEd/Private/Kismet2/KismetDebugUtilities.cpp
@@ -247,12 +247,6 @@ void FKismetDebugUtilities::OnScriptException(const UObject* ActiveObject, const
 
 	checkSlow(ActiveObject != nullptr);
 
-	// Ignore script exceptions for preview actors
-	if(FActorEditorUtils::IsAPreviewOrInactiveActor(Cast<const AActor>(ActiveObject)))
-	{
-		return;
-	}
-	
 	UClass* ClassContainingCode = FindClassForNode(ActiveObject, StackFrame.Node);
 	UBlueprint* BlueprintObj = (ClassContainingCode ? Cast<UBlueprint>(ClassContainingCode->ClassGeneratedBy) : nullptr);
 	if (BlueprintObj)
@@ -260,6 +254,19 @@ void FKismetDebugUtilities::OnScriptException(const UObject* ActiveObject, const
 		const FBlueprintExceptionInfo* ExceptionInfo = &Info;
 		bool bResetObjectBeingDebuggedWhenFinished = false;
 		UObject* ObjectBeingDebugged = BlueprintObj->GetObjectBeingDebugged();
+
+		auto IsAPreviewOrInactiveObject = [](const UObject* InObject)
+		{
+			UWorld* World = InObject ? InObject->GetWorld() : nullptr;
+			return World && (World->WorldType == EWorldType::EditorPreview || World->WorldType == EWorldType::Inactive);
+		};
+
+		// Ignore script exceptions for preview objects that are not already being debugged
+		if (IsAPreviewOrInactiveObject(ActiveObject) && ObjectBeingDebugged != ActiveObject)
+		{
+			return;
+		}
+
 		UObject* SavedObjectBeingDebugged = ObjectBeingDebugged;
 		UWorld* WorldBeingDebugged = BlueprintObj->GetWorldBeingDebugged();
 		const FString& PathToDebug = BlueprintObj->GetObjectPathToDebug();
