Improvements to anim node property tracing and display in the animation blueprint editor:
- Trace selected animation for ChooserPlayer node.
- Make debug bubbles in rewind debugger show all traced node properties, except Name.
- Make sequence player nodes display the current sequence in the debug bubble.
- Trace assets, such as Sequence or BlendSpace with the Key: Asset, and enable display of the Asset column in the AnimGraph rewind debugger details.

diff --git a/Engine/Plugins/Animation/GameplayInsights/Source/GameplayInsights/Private/AnimationProvider.cpp b/Engine/Plugins/Animation/GameplayInsights/Source/GameplayInsights/Private/AnimationProvider.cpp
--- a/Engine/Plugins/Animation/GameplayInsights/Source/GameplayInsights/Private/AnimationProvider.cpp
+++ b/Engine/Plugins/Animation/GameplayInsights/Source/GameplayInsights/Private/AnimationProvider.cpp
@@ -497,7 +497,7 @@ FText FAnimationProvider::FormatNodeValue(const FAnimNodeValueMessage& InMessage
 	case EAnimNodeValueType::Object:
 	{
 		const FObjectInfo& ObjectInfo = GameplayProvider.GetObjectInfo(InMessage.Value.Object.Value);
-		Text = FText::FromString(ObjectInfo.PathName);
+		Text = FText::FromString(ObjectInfo.Name);
 		break;
 	}
 	case EAnimNodeValueType::Class:
diff --git a/Engine/Plugins/Animation/GameplayInsights/Source/GameplayInsights/Private/SAnimGraphSchematicView.cpp b/Engine/Plugins/Animation/GameplayInsights/Source/GameplayInsights/Private/SAnimGraphSchematicView.cpp
--- a/Engine/Plugins/Animation/GameplayInsights/Source/GameplayInsights/Private/SAnimGraphSchematicView.cpp
+++ b/Engine/Plugins/Animation/GameplayInsights/Source/GameplayInsights/Private/SAnimGraphSchematicView.cpp
@@ -44,6 +44,7 @@ namespace AnimGraphSchematicColumns
 {
 	static const FName Type("Type");
 	static const FName Name("Name");
+	static const FName Asset("Asset");
 	static const FName Weight("Weight");
 	static const FName RootMotionWeight("Root Motion Weight");
 };
@@ -393,6 +394,7 @@ void SAnimGraphSchematicView::Construct(const FArguments& InArgs, uint64 InAnimI
 	// Make default columns
 	Columns.Add(AnimGraphSchematicColumns::Type, { 0, true });
 	Columns.Add(AnimGraphSchematicColumns::Name, { 1, true });
+	Columns.Add(AnimGraphSchematicColumns::Asset, { 2, true });
 	Columns.Add(AnimGraphSchematicColumns::Weight, { 99999, true });
 	Columns.Add(AnimGraphSchematicColumns::RootMotionWeight, { 100000, false });
 
diff --git a/Engine/Plugins/Animation/GameplayInsights/Source/RewindDebugger/Private/RewindDebugger.cpp b/Engine/Plugins/Animation/GameplayInsights/Source/RewindDebugger/Private/RewindDebugger.cpp
--- a/Engine/Plugins/Animation/GameplayInsights/Source/RewindDebugger/Private/RewindDebugger.cpp
+++ b/Engine/Plugins/Animation/GameplayInsights/Source/RewindDebugger/Private/RewindDebugger.cpp
@@ -1141,8 +1141,12 @@ void FRewindDebugger::Tick(float DeltaTime)
 																	TRACE_CPUPROFILER_EVENT_SCOPE(AnimGraphNodeValues);
 																	InNodeValuesTimeline.EnumerateEvents(InGraphStartTime, InGraphEndTime, [AnimationProvider, &DebugData](double InStartTime, double InEndTime, uint32 InDepth, const FAnimNodeValueMessage& InMessage)
 																	{
-																		FText Text = AnimationProvider->FormatNodeKeyValue(InMessage);
-																		DebugData.RecordNodeValue(InMessage.NodeId, Text.ToString());
+																		// don't send "Name" Node value for display in the graph
+																		if (FPlatformString::Strcmp(InMessage.Key, TEXT("Name")) != 0)
+																		{
+																			FText Text = AnimationProvider->FormatNodeKeyValue(InMessage);
+																			DebugData.RecordNodeValue(InMessage.NodeId, Text.ToString());
+																		}
 																		return TraceServices::EEventEnumerate::Continue;
 																	});
 																});
diff --git a/Engine/Plugins/Experimental/Chooser/Source/Chooser/Internal/Chooser.h b/Engine/Plugins/Experimental/Chooser/Source/Chooser/Internal/Chooser.h
--- a/Engine/Plugins/Experimental/Chooser/Source/Chooser/Internal/Chooser.h
+++ b/Engine/Plugins/Experimental/Chooser/Source/Chooser/Internal/Chooser.h
@@ -110,6 +110,8 @@ struct CHOOSER_API FEvaluateChooser : public FObjectChooserBase
 
 	virtual UObject* ChooseObject(FChooserEvaluationContext& Context) const final override;
 	virtual EIteratorStatus ChooseMulti(FChooserEvaluationContext &Context, FObjectChooserIteratorCallback Callback) const final override;
+	virtual void GetDebugName(FString& OutDebugName) const override;
+	
 	public:
 
 	FEvaluateChooser() : Chooser(nullptr) {}
diff --git a/Engine/Plugins/Experimental/Chooser/Source/Chooser/Private/AnimNode_ChooserPlayer.cpp b/Engine/Plugins/Experimental/Chooser/Source/Chooser/Private/AnimNode_ChooserPlayer.cpp
--- a/Engine/Plugins/Experimental/Chooser/Source/Chooser/Private/AnimNode_ChooserPlayer.cpp
+++ b/Engine/Plugins/Experimental/Chooser/Source/Chooser/Private/AnimNode_ChooserPlayer.cpp
@@ -161,6 +161,12 @@ void FAnimNode_ChooserPlayer::UpdateAssetPlayer(const FAnimationUpdateContext& C
 	}
 
 	FAnimNode_BlendStack_Standalone::UpdateAssetPlayer(Context);
+
+	FString Name;
+	Chooser.Get<FObjectChooserBase>().GetDebugName(Name);
+	
+	TRACE_ANIM_NODE_VALUE(Context, TEXT("Name"), ToCStr(Name));
+	TRACE_ANIM_NODE_VALUE(Context, TEXT("Asset"), CurrentAsset);
 }
 
 void FAnimNode_ChooserPlayer::Evaluate_AnyThread(FPoseContext& Output)
diff --git a/Engine/Plugins/Experimental/Chooser/Source/Chooser/Private/Chooser.cpp b/Engine/Plugins/Experimental/Chooser/Source/Chooser/Private/Chooser.cpp
--- a/Engine/Plugins/Experimental/Chooser/Source/Chooser/Private/Chooser.cpp
+++ b/Engine/Plugins/Experimental/Chooser/Source/Chooser/Private/Chooser.cpp
@@ -22,7 +22,9 @@ void UChooserTable::PostLoad()
 #if WITH_EDITOR
 	CachedPreviousOutputObjectType = OutputObjectType;
 	CachedPreviousResultType = ResultType;
+#endif
 
+#if WITH_EDITORONLY_DATA
 	// convert old data if it exists
 
 	if (ContextObjectType_DEPRECATED)
@@ -322,4 +324,12 @@ UObject* FEvaluateChooser::ChooseObject(FChooserEvaluationContext& Context) cons
 FObjectChooserBase::EIteratorStatus FEvaluateChooser::ChooseMulti(FChooserEvaluationContext& Context, FObjectChooserIteratorCallback Callback) const
 {
 	return UChooserTable::EvaluateChooser(Context, Chooser, Callback);
-}
\ No newline at end of file
+}
+
+void FEvaluateChooser::GetDebugName(FString& OutDebugName) const
+{
+	if (Chooser)
+	{
+		OutDebugName = Chooser.GetName();
+	}
+}
diff --git a/Engine/Plugins/Experimental/Chooser/Source/Chooser/Public/IObjectChooser.h b/Engine/Plugins/Experimental/Chooser/Source/Chooser/Public/IObjectChooser.h
--- a/Engine/Plugins/Experimental/Chooser/Source/Chooser/Public/IObjectChooser.h
+++ b/Engine/Plugins/Experimental/Chooser/Source/Chooser/Public/IObjectChooser.h
@@ -71,4 +71,6 @@ public:
 		}
 		return EIteratorStatus::Continue;
 	}
+
+	virtual void GetDebugName(FString& OutDebugName) const {};
 };
\ No newline at end of file
diff --git a/Engine/Plugins/Experimental/Chooser/Source/ProxyTable/Internal/LookupProxy.h b/Engine/Plugins/Experimental/Chooser/Source/ProxyTable/Internal/LookupProxy.h
--- a/Engine/Plugins/Experimental/Chooser/Source/ProxyTable/Internal/LookupProxy.h
+++ b/Engine/Plugins/Experimental/Chooser/Source/ProxyTable/Internal/LookupProxy.h
@@ -38,6 +38,8 @@ struct PROXYTABLE_API FLookupProxy : public FObjectChooserBase
 	FLookupProxy();
 	
 	virtual void Compile(IHasContextClass* HasContext, bool bForce) override;
+
+	virtual void GetDebugName(FString& OutName) const override;
 	
 	public:
 	
diff --git a/Engine/Plugins/Experimental/Chooser/Source/ProxyTable/Private/LookupProxy.cpp b/Engine/Plugins/Experimental/Chooser/Source/ProxyTable/Private/LookupProxy.cpp
--- a/Engine/Plugins/Experimental/Chooser/Source/ProxyTable/Private/LookupProxy.cpp
+++ b/Engine/Plugins/Experimental/Chooser/Source/ProxyTable/Private/LookupProxy.cpp
@@ -116,3 +116,11 @@ void FLookupProxy::Compile(IHasContextClass* HasContext, bool bForce)
 		}
 	}
 }
+
+void FLookupProxy::GetDebugName(FString& OutName) const
+{
+	if (Proxy)
+	{
+		OutName = Proxy.GetName();
+	}
+}
diff --git a/Engine/Source/Editor/AnimationBlueprintEditor/Private/AnimationNodes/SAnimationGraphNode.cpp b/Engine/Source/Editor/AnimationBlueprintEditor/Private/AnimationNodes/SAnimationGraphNode.cpp
--- a/Engine/Source/Editor/AnimationBlueprintEditor/Private/AnimationNodes/SAnimationGraphNode.cpp
+++ b/Engine/Source/Editor/AnimationBlueprintEditor/Private/AnimationNodes/SAnimationGraphNode.cpp
@@ -177,9 +177,24 @@ void SAnimationGraphNode::GetNodeInfoPopups(FNodeInfoContext* Context, TArray<FG
 					// reverse node index temporarily because of a bug in NodeGuidToIndexMap
 					AnimNodeIndex = Class->GetAnimNodeProperties().Num() - AnimNodeIndex - 1;
 
-					if (FAnimBlueprintDebugData::FNodeValue* DebugInfo = Class->GetAnimBlueprintDebugData().NodeValuesThisFrame.FindByPredicate([AnimNodeIndex](const FAnimBlueprintDebugData::FNodeValue& InValue){ return InValue.NodeID == AnimNodeIndex; }))
+					FString PopupText;
+					for(auto & NodeValue : Class->GetAnimBlueprintDebugData().NodeValuesThisFrame)
 					{
-						Popups.Emplace(nullptr, Color, DebugInfo->Text);
+						if (NodeValue.NodeID == AnimNodeIndex)
+						{
+							if (PopupText.IsEmpty())
+							{
+								PopupText = NodeValue.Text;
+							}
+							else
+							{
+								PopupText = FString::Format(TEXT("{0}\n{1}"), {PopupText, NodeValue.Text});
+							}
+						}
+					}
+					if (!PopupText.IsEmpty())
+					{
+						Popups.Emplace(nullptr, Color, PopupText);
 					}
 				}
 			}
diff --git a/Engine/Source/Editor/AnimationBlueprintEditor/Private/AnimationNodes/SGraphNodeSequencePlayer.cpp b/Engine/Source/Editor/AnimationBlueprintEditor/Private/AnimationNodes/SGraphNodeSequencePlayer.cpp
--- a/Engine/Source/Editor/AnimationBlueprintEditor/Private/AnimationNodes/SGraphNodeSequencePlayer.cpp
+++ b/Engine/Source/Editor/AnimationBlueprintEditor/Private/AnimationNodes/SGraphNodeSequencePlayer.cpp
@@ -7,6 +7,7 @@
 #include "AnimGraphNode_SequencePlayer.h"
 #include "Animation/AnimBlueprint.h"
 #include "Animation/AnimBlueprintGeneratedClass.h"
+#include "Animation/AnimInstance.h"
 #include "Animation/AnimNode_SequencePlayer.h"
 #include "Animation/AnimSequenceBase.h"
 #include "Animation/AnimationAsset.h"
@@ -70,6 +71,50 @@ void SGraphNodeSequencePlayer::Construct(const FArguments& InArgs, UAnimGraphNod
 
 void SGraphNodeSequencePlayer::GetNodeInfoPopups(FNodeInfoContext* Context, TArray<FGraphInformationPopupInfo>& Popups) const
 {
+	SGraphNodeK2Base::GetNodeInfoPopups(Context, Popups);
+
+	UAnimBlueprint* AnimBlueprint = Cast<UAnimBlueprint>(FBlueprintEditorUtils::FindBlueprintForNode(GraphNode));
+	if(AnimBlueprint)
+	{
+		UAnimInstance* ActiveObject = Cast<UAnimInstance>(AnimBlueprint->GetObjectBeingDebugged());
+		UAnimBlueprintGeneratedClass* Class = AnimBlueprint->GetAnimBlueprintGeneratedClass();
+
+		const FLinearColor Color(1.f, 0.5f, 0.25f);
+
+		// Display various types of debug data
+		if ((ActiveObject != NULL) && (Class != NULL))
+		{
+			if (Class->GetAnimNodeProperties().Num())
+			{
+				if(int32* NodeIndexPtr = Class->GetAnimBlueprintDebugData().NodePropertyToIndexMap.Find(TWeakObjectPtr<UAnimGraphNode_Base>(Cast<UAnimGraphNode_Base>(GraphNode))))
+				{
+					int32 AnimNodeIndex = *NodeIndexPtr;
+					// reverse node index temporarily because of a bug in NodeGuidToIndexMap
+					AnimNodeIndex = Class->GetAnimNodeProperties().Num() - AnimNodeIndex - 1;
+
+					FString PopupText;
+					for(auto & NodeValue : Class->GetAnimBlueprintDebugData().NodeValuesThisFrame)
+					{
+						if (NodeValue.NodeID == AnimNodeIndex && NodeValue.Text.Contains("Asset ="))
+						{
+							if (PopupText.IsEmpty())
+							{
+								PopupText = NodeValue.Text;
+							}
+							else
+							{
+								PopupText = FString::Format(TEXT("{0}\n{1}"), {PopupText, NodeValue.Text});
+							}
+						}
+					}
+					if (!PopupText.IsEmpty())
+					{
+						Popups.Emplace(nullptr, Color, PopupText);
+					}
+				}
+			}
+		}
+	}
 }
 
 FText SGraphNodeSequencePlayer::GetPositionTooltip() const
diff --git a/Engine/Source/Runtime/AnimGraphRuntime/Private/AnimNodes/AnimNode_BlendSpaceEvaluator.cpp b/Engine/Source/Runtime/AnimGraphRuntime/Private/AnimNodes/AnimNode_BlendSpaceEvaluator.cpp
--- a/Engine/Source/Runtime/AnimGraphRuntime/Private/AnimNodes/AnimNode_BlendSpaceEvaluator.cpp
+++ b/Engine/Source/Runtime/AnimGraphRuntime/Private/AnimNodes/AnimNode_BlendSpaceEvaluator.cpp
@@ -23,7 +23,7 @@ void FAnimNode_BlendSpaceEvaluator::UpdateAssetPlayer(const FAnimationUpdateCont
 	UpdateInternal(Context);
 
 	TRACE_ANIM_NODE_VALUE(Context, TEXT("Name"), GetBlendSpace() ? *GetBlendSpace()->GetName() : TEXT("None"));
-	TRACE_ANIM_NODE_VALUE(Context, TEXT("Blend Space"), GetBlendSpace());
+	TRACE_ANIM_NODE_VALUE(Context, TEXT("Asset"), GetBlendSpace());
 	TRACE_ANIM_NODE_VALUE(Context, TEXT("Playback Time"), InternalTimeAccumulator);
 }
 
diff --git a/Engine/Source/Runtime/AnimGraphRuntime/Private/AnimNodes/AnimNode_BlendSpaceGraphBase.cpp b/Engine/Source/Runtime/AnimGraphRuntime/Private/AnimNodes/AnimNode_BlendSpaceGraphBase.cpp
--- a/Engine/Source/Runtime/AnimGraphRuntime/Private/AnimNodes/AnimNode_BlendSpaceGraphBase.cpp
+++ b/Engine/Source/Runtime/AnimGraphRuntime/Private/AnimNodes/AnimNode_BlendSpaceGraphBase.cpp
@@ -74,7 +74,7 @@ void FAnimNode_BlendSpaceGraphBase::UpdateInternal(const FAnimationUpdateContext
 
 	TRACE_BLENDSPACE(Context, *this);
 	TRACE_ANIM_NODE_VALUE(Context, TEXT("Name"), *BlendSpace->GetName());
-	TRACE_ANIM_NODE_VALUE(Context, TEXT("Blend Space"), BlendSpace);
+	TRACE_ANIM_NODE_VALUE(Context, TEXT("Asset"), BlendSpace);
 }
 
 void FAnimNode_BlendSpaceGraphBase::Update_AnyThread(const FAnimationUpdateContext& Context)
diff --git a/Engine/Source/Runtime/AnimGraphRuntime/Private/AnimNodes/AnimNode_BlendSpacePlayer.cpp b/Engine/Source/Runtime/AnimGraphRuntime/Private/AnimNodes/AnimNode_BlendSpacePlayer.cpp
--- a/Engine/Source/Runtime/AnimGraphRuntime/Private/AnimNodes/AnimNode_BlendSpacePlayer.cpp
+++ b/Engine/Source/Runtime/AnimGraphRuntime/Private/AnimNodes/AnimNode_BlendSpacePlayer.cpp
@@ -117,7 +117,7 @@ void FAnimNode_BlendSpacePlayerBase::UpdateInternal(const FAnimationUpdateContex
 
 	TRACE_BLENDSPACE_PLAYER(Context, *this);
 	TRACE_ANIM_NODE_VALUE(Context, TEXT("Name"), CurrentBlendSpace ? *CurrentBlendSpace->GetName() : TEXT("None"));
-	TRACE_ANIM_NODE_VALUE(Context, TEXT("Blend Space"), CurrentBlendSpace);
+	TRACE_ANIM_NODE_VALUE(Context, TEXT("Asset"), CurrentBlendSpace);
 	TRACE_ANIM_NODE_VALUE(Context, TEXT("Playback Time"), InternalTimeAccumulator);
 }
 
diff --git a/Engine/Source/Runtime/Engine/Private/Animation/AnimNode_SequencePlayer.cpp b/Engine/Source/Runtime/Engine/Private/Animation/AnimNode_SequencePlayer.cpp
--- a/Engine/Source/Runtime/Engine/Private/Animation/AnimNode_SequencePlayer.cpp
+++ b/Engine/Source/Runtime/Engine/Private/Animation/AnimNode_SequencePlayer.cpp
@@ -112,7 +112,7 @@ void FAnimNode_SequencePlayerBase::UpdateAssetPlayer(const FAnimationUpdateConte
 
 	TRACE_ANIM_SEQUENCE_PLAYER(Context, *this);
 	TRACE_ANIM_NODE_VALUE(Context, TEXT("Name"), CurrentSequence != nullptr ? CurrentSequence->GetFName() : NAME_None);
-	TRACE_ANIM_NODE_VALUE(Context, TEXT("Sequence"), CurrentSequence);
+	TRACE_ANIM_NODE_VALUE(Context, TEXT("Asset"), CurrentSequence);
 	TRACE_ANIM_NODE_VALUE(Context, TEXT("Playback Time"), InternalTimeAccumulator);
 }
 
