Fix for debug draw on blendspace player nodes when debugging parent anim bps

diff --git a/Engine/Source/Editor/AnimationBlueprintEditor/Private/AnimationNodes/SBlendSpacePreview.cpp b/Engine/Source/Editor/AnimationBlueprintEditor/Private/AnimationNodes/SBlendSpacePreview.cpp
--- a/Engine/Source/Editor/AnimationBlueprintEditor/Private/AnimationNodes/SBlendSpacePreview.cpp
+++ b/Engine/Source/Editor/AnimationBlueprintEditor/Private/AnimationNodes/SBlendSpacePreview.cpp
@@ -90,19 +90,26 @@ bool SBlendSpacePreview::GetBlendSpaceInfo(TWeakObjectPtr<const UBlendSpace>& Ou
 			{
 				if (UAnimBlueprintGeneratedClass* Class = Cast<UAnimBlueprintGeneratedClass>(ActiveObject->GetClass()))
 				{
-					if(int32* NodeIndexPtr = Class->GetAnimBlueprintDebugData().NodePropertyToIndexMap.Find(Node))
+					if (UAnimBlueprint* AnimBlueprint = Cast<UAnimBlueprint>(Class->ClassGeneratedBy))
 					{
-						int32 AnimNodeIndex = *NodeIndexPtr;
-						// reverse node index temporarily because of a bug in NodeGuidToIndexMap
-						AnimNodeIndex = Class->GetAnimNodeProperties().Num() - AnimNodeIndex - 1;
+						UAnimBlueprint* RootAnimBP = UAnimBlueprint::FindRootAnimBlueprint(AnimBlueprint);
+						AnimBlueprint = RootAnimBP ? RootAnimBP : AnimBlueprint;
+					
+						const FAnimBlueprintDebugData& DebugData = AnimBlueprint->GetAnimBlueprintGeneratedClass()->GetAnimBlueprintDebugData();
+						if(const int32* NodeIndexPtr = DebugData.NodePropertyToIndexMap.Find(Node))
+						{
+							int32 AnimNodeIndex = *NodeIndexPtr;
+							// reverse node index temporarily because of a bug in NodeGuidToIndexMap
+							AnimNodeIndex = Class->GetAnimNodeProperties().Num() - AnimNodeIndex - 1;
 
-						if (FAnimBlueprintDebugData::FBlendSpacePlayerRecord* DebugInfo = Class->GetAnimBlueprintDebugData().BlendSpacePlayerRecordsThisFrame.FindByPredicate(
+							if (const FAnimBlueprintDebugData::FBlendSpacePlayerRecord* DebugInfo = DebugData.BlendSpacePlayerRecordsThisFrame.FindByPredicate(
 							[AnimNodeIndex](const FAnimBlueprintDebugData::FBlendSpacePlayerRecord& InRecord){ return InRecord.NodeID == AnimNodeIndex; }))
-						{
-							OutBlendSpace = DebugInfo->BlendSpace.Get();
-							OutPosition = DebugInfo->Position;
-							OutFilteredPosition = DebugInfo->FilteredPosition;
-							return true;
+							{
+								OutBlendSpace = DebugInfo->BlendSpace.Get();
+								OutPosition = DebugInfo->Position;
+								OutFilteredPosition = DebugInfo->FilteredPosition;
+								return true;
+							}
 						}
 					}
 				}
