Removing dependency on private Persona module in AnimationBlueprintEditor module.

diff --git a/Engine/Source/Editor/AnimationBlueprintEditor/AnimationBlueprintEditor.Build.cs b/Engine/Source/Editor/AnimationBlueprintEditor/AnimationBlueprintEditor.Build.cs
--- a/Engine/Source/Editor/AnimationBlueprintEditor/AnimationBlueprintEditor.Build.cs
+++ b/Engine/Source/Editor/AnimationBlueprintEditor/AnimationBlueprintEditor.Build.cs
@@ -13,12 +13,6 @@ public class AnimationBlueprintEditor : ModuleRules
 			}
 		);
 
-		PrivateIncludePaths.AddRange(
-			new string[] {
-				System.IO.Path.Combine(GetModuleDirectory("Persona"), "Private"),
-			}
-		);
-
 		PrivateDependencyModuleNames.AddRange(
 			new string[] {
                 "AppFramework",
diff --git a/Engine/Source/Editor/AnimationBlueprintEditor/Private/AnimationBlueprintEditor.cpp b/Engine/Source/Editor/AnimationBlueprintEditor/Private/AnimationBlueprintEditor.cpp
--- a/Engine/Source/Editor/AnimationBlueprintEditor/Private/AnimationBlueprintEditor.cpp
+++ b/Engine/Source/Editor/AnimationBlueprintEditor/Private/AnimationBlueprintEditor.cpp
@@ -44,7 +44,6 @@
 #include "AnimationBlueprintEditorModule.h"
 #include "AnimationBlueprintEditorSettings.h"
 #include "AnimationBlueprintInterfaceEditorMode.h"
-#include "AnimationEditorPreviewScene.h"
 #include "AnimationEditorUtils.h"
 #include "AnimationGraph.h"
 #include "AnimationGraphSchema.h"
@@ -100,11 +99,7 @@
 // Hide related nodes feature
 #include "AnimationBlueprintToolMenuContext.h"
 #include "DetailLayoutBuilder.h"
-#include "DetailWidgetRow.h"
-#include "PersonaPreviewSceneDescription.h"
 #include "Preferences/AnimationBlueprintEditorOptions.h"
-#include "Preferences/PersonaOptions.h"
-#include "PropertyCustomizationHelpers.h"
 #include "PropertyEditorDelegates.h"
 #include "SBlueprintEditorToolbar.h"
 #include "SKismetInspector.h"
@@ -425,11 +420,8 @@ void FAnimationBlueprintEditor::InitAnimationBlueprintEditor(const EToolkitMode:
 
 	GetToolkitCommands()->Append(FPlayWorldCommands::GlobalPlayWorldActions.ToSharedRef());
 
-	FPersonaToolkitArgs PersonaToolkitArgs;
-	PersonaToolkitArgs.OnPreviewSceneSettingsCustomized = FOnPreviewSceneSettingsCustomized::FDelegate::CreateSP(this, &FAnimationBlueprintEditor::HandleOnPreviewSceneSettingsCustomized);
-
 	FPersonaModule& PersonaModule = FModuleManager::GetModuleChecked<FPersonaModule>("Persona");
-	PersonaToolkit = PersonaModule.CreatePersonaToolkit(InAnimBlueprint, PersonaToolkitArgs);
+	PersonaToolkit = PersonaModule.CreatePersonaToolkit(InAnimBlueprint, FPersonaToolkitArgs());
 
 	PersonaToolkit->GetPreviewScene()->SetDefaultAnimationMode(EPreviewSceneDefaultAnimationMode::AnimationBlueprint);
 	PersonaToolkit->GetPreviewScene()->RegisterOnPreviewMeshChanged(FOnPreviewMeshChanged::CreateSP(this, &FAnimationBlueprintEditor::HandlePreviewMeshChanged));
@@ -2219,129 +2211,6 @@ void FAnimationBlueprintEditor::HandlePoseWatchSelectedNodes()
 	}
 }
 
-void FAnimationBlueprintEditor::HandleOnPreviewSceneSettingsCustomized(IDetailLayoutBuilder& DetailBuilder)
-{
-	const TSharedRef<IPropertyHandle> PreviewAnimationBlueprintProperty = DetailBuilder.GetProperty(GET_MEMBER_NAME_CHECKED(UPersonaPreviewSceneDescription, PreviewAnimationBlueprint));
-	const TSharedRef<IPropertyHandle> ApplicationMethodProperty = DetailBuilder.GetProperty(GET_MEMBER_NAME_CHECKED(UPersonaPreviewSceneDescription, ApplicationMethod));
-	const TSharedRef<IPropertyHandle> LinkedAnimGraphTagProperty = DetailBuilder.GetProperty(GET_MEMBER_NAME_CHECKED(UPersonaPreviewSceneDescription, LinkedAnimGraphTag));
-	
-	// customize "Preview Animation Blueprint" for animation blueprint classes
-	DetailBuilder.EditCategory("Animation Blueprint")
-	.AddProperty(PreviewAnimationBlueprintProperty)
-	.CustomWidget()
-	.NameContent()
-	[
-		SNew(SVerticalBox)
-		+SVerticalBox::Slot()
-		.AutoHeight()
-		[
-			PreviewAnimationBlueprintProperty->CreatePropertyNameWidget()
-		]
-	]
-	.ValueContent()
-	.MaxDesiredWidth(250.0f)
-	.MinDesiredWidth(250.0f)
-	[
-		SNew(SObjectPropertyEntryBox)
-		.AllowedClass(UAnimBlueprint::StaticClass())
-		.PropertyHandle(PreviewAnimationBlueprintProperty)
-		.OnShouldFilterAsset(this, &FAnimationBlueprintEditor::HandleShouldFilterAsset, FName("TargetSkeleton"), false)
-		.OnObjectChanged(this, &FAnimationBlueprintEditor::HandlePreviewAnimBlueprintChanged)
-		.ThumbnailPool(DetailBuilder.GetThumbnailPool())
-	];
-
-	ApplicationMethodProperty->SetOnPropertyValueChanged(FSimpleDelegate::CreateLambda([this]()
-	{
-		FScopedTransaction Transaction(LOCTEXT("SetAnimationBlueprintApplicationMethod", "Set Application Method"));
-		
-		const TSharedRef<FAnimationEditorPreviewScene> LocalPreviewScene = StaticCastSharedRef<FAnimationEditorPreviewScene>(PersonaToolkit->GetPreviewScene());
-		const UPersonaPreviewSceneDescription* PersonaPreviewSceneDescription = LocalPreviewScene->GetPreviewSceneDescription();
-		PersonaToolkit->GetAnimBlueprint()->SetPreviewAnimationBlueprintApplicationMethod(PersonaPreviewSceneDescription->ApplicationMethod);
-		LocalPreviewScene->SetPreviewAnimationBlueprint(PersonaPreviewSceneDescription->PreviewAnimationBlueprint.Get(), PersonaToolkit->GetAnimBlueprint());
-	}));
-
-	DetailBuilder.EditCategory("Animation Blueprint")
-	.AddProperty(ApplicationMethodProperty)
-	.IsEnabled(MakeAttributeLambda([this]()
-	{
-		TSharedRef<FAnimationEditorPreviewScene> LocalPreviewScene = StaticCastSharedRef<FAnimationEditorPreviewScene>(PersonaToolkit->GetPreviewScene());
-		UPersonaPreviewSceneDescription* PersonaPreviewSceneDescription = LocalPreviewScene->GetPreviewSceneDescription();
-		
-		return PersonaPreviewSceneDescription->PreviewAnimationBlueprint.IsValid();
-	}));
-
-	LinkedAnimGraphTagProperty->SetOnPropertyValueChanged(FSimpleDelegate::CreateLambda([this]()
-	{
-		FScopedTransaction Transaction(LOCTEXT("SetAnimationBlueprintTag", "Set Linked Anim Graph Tag"));
-		
-		TSharedRef<FAnimationEditorPreviewScene> LocalPreviewScene = StaticCastSharedRef<FAnimationEditorPreviewScene>(PersonaToolkit->GetPreviewScene());
-		UPersonaPreviewSceneDescription* PersonaPreviewSceneDescription = LocalPreviewScene->GetPreviewSceneDescription();
-		PersonaToolkit->GetAnimBlueprint()->SetPreviewAnimationBlueprintTag(PersonaPreviewSceneDescription->LinkedAnimGraphTag);
-		LocalPreviewScene->SetPreviewAnimationBlueprint(PersonaPreviewSceneDescription->PreviewAnimationBlueprint.Get(), PersonaToolkit->GetAnimBlueprint());
-	}));
-
-	DetailBuilder.EditCategory("Animation Blueprint")
-	.AddProperty(LinkedAnimGraphTagProperty)
-	.IsEnabled(MakeAttributeLambda([this]()
-	{
-		TSharedRef<FAnimationEditorPreviewScene> LocalPreviewScene = StaticCastSharedRef<FAnimationEditorPreviewScene>(PersonaToolkit->GetPreviewScene());
-		UPersonaPreviewSceneDescription* PersonaPreviewSceneDescription = LocalPreviewScene->GetPreviewSceneDescription();
-		
-		return PersonaPreviewSceneDescription->PreviewAnimationBlueprint.IsValid() && PersonaPreviewSceneDescription->ApplicationMethod == EPreviewAnimationBlueprintApplicationMethod::LinkedAnimGraph;
-	}));
-	
-	if (PersonaToolkit->GetAnimBlueprint())
-	{
-		PersonaToolkit->GetAnimBlueprint()->OnCompiled().AddSP(this, &FAnimationBlueprintEditor::HandleAnimBlueprintCompiled);
-	}
-}
-
-void FAnimationBlueprintEditor::HandleAnimBlueprintCompiled(UBlueprint* Blueprint) const
-{
-	// Only re-initialize controller if we are not debugging an external instance.
-	// If we switch at this point then we will disconnect from the external instance
-	const TSharedRef<FAnimationEditorPreviewScene> AnimPreviewScene = StaticCastSharedRef<FAnimationEditorPreviewScene>(GetPreviewScene());
-	if(AnimPreviewScene->GetPreviewMeshComponent()->PreviewInstance == nullptr || AnimPreviewScene->GetPreviewMeshComponent()->PreviewInstance->GetDebugSkeletalMeshComponent() == nullptr)
-	{
-		UPersonaPreviewSceneDescription* PersonaPreviewSceneDescription = AnimPreviewScene->GetPreviewSceneDescription();
-		PersonaPreviewSceneDescription->PreviewControllerInstance->UninitializeView(PersonaPreviewSceneDescription, &GetPreviewScene().Get());
-		PersonaPreviewSceneDescription->PreviewControllerInstance->InitializeView(PersonaPreviewSceneDescription, &GetPreviewScene().Get());
-	}
-}
-
-void FAnimationBlueprintEditor::HandlePreviewAnimBlueprintChanged(const FAssetData& InAssetData) const   
-{
-	UAnimBlueprint* NewAnimBlueprint = Cast<UAnimBlueprint>(InAssetData.GetAsset());
-	PersonaToolkit->SetPreviewAnimationBlueprint(NewAnimBlueprint);
-}
-
-bool FAnimationBlueprintEditor::HandleShouldFilterAsset(
-	const FAssetData& InAssetData,
-	FName InTag,
-	bool bCanUseDifferentSkeleton) const
-{
-	if (bCanUseDifferentSkeleton && GetDefault<UPersonaOptions>()->bAllowPreviewMeshCollectionsToSelectFromDifferentSkeletons)
-	{
-		return false;
-	}
-
-	const TSharedRef<FAnimationEditorPreviewScene> AnimPreviewScene = StaticCastSharedRef<FAnimationEditorPreviewScene>(GetPreviewScene());
-	const UPersonaPreviewSceneDescription* PersonaPreviewSceneDescription = AnimPreviewScene->GetPreviewSceneDescription();
-	if(!PersonaPreviewSceneDescription->PreviewMesh.IsValid())
-	{
-		return false;
-	}
-	
-	const USkeleton* Skeleton = PersonaPreviewSceneDescription->PreviewMesh->GetSkeleton();
-	const FString SkeletonTag = InAssetData.GetTagValueRef<FString>(InTag);
-	if (Skeleton && Skeleton->IsCompatibleForEditor(SkeletonTag))
-	{
-		return false;
-	}
-
-	return true;
-}
-
 void FAnimationBlueprintEditor::RemoveAllSelectionPoseWatches()
 {
 	TSharedPtr<SGraphEditor> FocusedGraphEd = FocusedGraphEdPtr.Pin();
diff --git a/Engine/Source/Editor/AnimationBlueprintEditor/Private/AnimationBlueprintEditor.h b/Engine/Source/Editor/AnimationBlueprintEditor/Private/AnimationBlueprintEditor.h
--- a/Engine/Source/Editor/AnimationBlueprintEditor/Private/AnimationBlueprintEditor.h
+++ b/Engine/Source/Editor/AnimationBlueprintEditor/Private/AnimationBlueprintEditor.h
@@ -339,12 +339,6 @@ private:
 	/** Enable/disable pose watch on selected nodes */
 	void HandlePoseWatchSelectedNodes();
 
-	/** Customize Preview Scene Settings details panel */
-	void HandleOnPreviewSceneSettingsCustomized(IDetailLayoutBuilder& DetailBuilder);
-	void HandlePreviewAnimBlueprintChanged(const FAssetData& InAssetData) const;
-	bool HandleShouldFilterAsset(const FAssetData& InAssetData, FName InTag, bool bCanUseDifferentSkeleton) const;
-	void HandleAnimBlueprintCompiled(UBlueprint* Blueprint) const;
-
 	/** Removes all pose watches created by selection from the current view */
 	void RemoveAllSelectionPoseWatches();
 
diff --git a/Engine/Source/Editor/Persona/Private/PreviewSceneCustomizations.cpp b/Engine/Source/Editor/Persona/Private/PreviewSceneCustomizations.cpp
--- a/Engine/Source/Editor/Persona/Private/PreviewSceneCustomizations.cpp
+++ b/Engine/Source/Editor/Persona/Private/PreviewSceneCustomizations.cpp
@@ -266,7 +266,88 @@ void FPreviewSceneDescriptionCustomization::CustomizeDetails(IDetailLayoutBuilde
 			.ThumbnailPool(DetailBuilder.GetThumbnailPool())
 		];
 	}
+
+	// Customize animation blueprint preview
+	TSharedRef<IPropertyHandle> PreviewAnimationBlueprintProperty = DetailBuilder.GetProperty(GET_MEMBER_NAME_CHECKED(UPersonaPreviewSceneDescription, PreviewAnimationBlueprint));
+	TSharedRef<IPropertyHandle> ApplicationMethodProperty = DetailBuilder.GetProperty(GET_MEMBER_NAME_CHECKED(UPersonaPreviewSceneDescription, ApplicationMethod));
+	TSharedRef<IPropertyHandle> LinkedAnimGraphTagProperty = DetailBuilder.GetProperty(GET_MEMBER_NAME_CHECKED(UPersonaPreviewSceneDescription, LinkedAnimGraphTag));
 	
+	if (PersonaToolkit.Pin()->GetContext() == UAnimBlueprint::StaticClass()->GetFName())
+	{
+		DetailBuilder.EditCategory("Animation Blueprint")
+		.AddProperty(PreviewAnimationBlueprintProperty)
+		.CustomWidget()
+		.NameContent()
+		[
+			SNew(SVerticalBox)
+			+SVerticalBox::Slot()
+			.AutoHeight()
+			[
+				PreviewAnimationBlueprintProperty->CreatePropertyNameWidget()
+			]
+		]
+		.ValueContent()
+		.MaxDesiredWidth(250.0f)
+		.MinDesiredWidth(250.0f)
+		[
+			SNew(SObjectPropertyEntryBox)
+			.AllowedClass(UAnimBlueprint::StaticClass())
+			.PropertyHandle(PreviewAnimationBlueprintProperty)
+			.OnShouldFilterAsset(this, &FPreviewSceneDescriptionCustomization::HandleShouldFilterAsset, FName("TargetSkeleton"), false)
+			.OnObjectChanged(this, &FPreviewSceneDescriptionCustomization::HandlePreviewAnimBlueprintChanged)
+			.ThumbnailPool(DetailBuilder.GetThumbnailPool())
+		];
+
+		ApplicationMethodProperty->SetOnPropertyValueChanged(FSimpleDelegate::CreateLambda([this]()
+		{
+			FScopedTransaction Transaction(LOCTEXT("SetAnimationBlueprintApplicationMethod", "Set Application Method"));
+
+			TSharedPtr<IPersonaToolkit> PinnedPersonaToolkit = PersonaToolkit.Pin();
+			TSharedRef<FAnimationEditorPreviewScene> LocalPreviewScene = StaticCastSharedRef<FAnimationEditorPreviewScene>(PinnedPersonaToolkit->GetPreviewScene());
+			UPersonaPreviewSceneDescription* PersonaPreviewSceneDescription = LocalPreviewScene->GetPreviewSceneDescription();
+			PinnedPersonaToolkit->GetAnimBlueprint()->SetPreviewAnimationBlueprintApplicationMethod(PersonaPreviewSceneDescription->ApplicationMethod);
+			LocalPreviewScene->SetPreviewAnimationBlueprint(PersonaPreviewSceneDescription->PreviewAnimationBlueprint.Get(), PinnedPersonaToolkit->GetAnimBlueprint());
+		}));
+
+		DetailBuilder.EditCategory("Animation Blueprint")
+		.AddProperty(ApplicationMethodProperty)
+		.IsEnabled(MakeAttributeLambda([this]()
+		{
+			TSharedPtr<IPersonaToolkit> PinnedPersonaToolkit = PersonaToolkit.Pin();
+			TSharedRef<FAnimationEditorPreviewScene> LocalPreviewScene = StaticCastSharedRef<FAnimationEditorPreviewScene>(PinnedPersonaToolkit->GetPreviewScene());
+			UPersonaPreviewSceneDescription* PersonaPreviewSceneDescription = LocalPreviewScene->GetPreviewSceneDescription();
+			
+			return PersonaPreviewSceneDescription->PreviewAnimationBlueprint.IsValid();
+		}));
+	
+		LinkedAnimGraphTagProperty->SetOnPropertyValueChanged(FSimpleDelegate::CreateLambda([this]()
+		{
+			FScopedTransaction Transaction(LOCTEXT("SetAnimationBlueprintTag", "Set Linked Anim Graph Tag"));
+
+			TSharedPtr<IPersonaToolkit> PinnedPersonaToolkit = PersonaToolkit.Pin();
+			TSharedRef<FAnimationEditorPreviewScene> LocalPreviewScene = StaticCastSharedRef<FAnimationEditorPreviewScene>(PinnedPersonaToolkit->GetPreviewScene());
+			UPersonaPreviewSceneDescription* PersonaPreviewSceneDescription = LocalPreviewScene->GetPreviewSceneDescription();
+			PinnedPersonaToolkit->GetAnimBlueprint()->SetPreviewAnimationBlueprintTag(PersonaPreviewSceneDescription->LinkedAnimGraphTag);
+			LocalPreviewScene->SetPreviewAnimationBlueprint(PersonaPreviewSceneDescription->PreviewAnimationBlueprint.Get(), PinnedPersonaToolkit->GetAnimBlueprint());
+		}));
+
+		DetailBuilder.EditCategory("Animation Blueprint")
+		.AddProperty(LinkedAnimGraphTagProperty)
+		.IsEnabled(MakeAttributeLambda([this]()
+		{
+			TSharedPtr<IPersonaToolkit> PinnedPersonaToolkit = PersonaToolkit.Pin();
+			TSharedRef<FAnimationEditorPreviewScene> LocalPreviewScene = StaticCastSharedRef<FAnimationEditorPreviewScene>(PinnedPersonaToolkit->GetPreviewScene());
+			UPersonaPreviewSceneDescription* PersonaPreviewSceneDescription = LocalPreviewScene->GetPreviewSceneDescription();
+			
+			return PersonaPreviewSceneDescription->PreviewAnimationBlueprint.IsValid() && PersonaPreviewSceneDescription->ApplicationMethod == EPreviewAnimationBlueprintApplicationMethod::LinkedAnimGraph;
+		}));
+	}
+	else
+	{
+		PreviewAnimationBlueprintProperty->MarkHiddenByCustomization();
+		ApplicationMethodProperty->MarkHiddenByCustomization();
+		LinkedAnimGraphTagProperty->MarkHiddenByCustomization();
+	}
 
 	//
 	// Physics section...
@@ -533,6 +614,25 @@ void FPreviewSceneDescriptionCustomization::HandleMeshChanged(const FAssetData&
 	PersonaToolkit.Pin()->SetPreviewMesh(NewPreviewMesh, false);
 }
 
+void FPreviewSceneDescriptionCustomization::HandlePreviewAnimBlueprintChanged(const FAssetData& InAssetData)
+{
+	UAnimBlueprint* NewAnimBlueprint = Cast<UAnimBlueprint>(InAssetData.GetAsset());
+	PersonaToolkit.Pin()->SetPreviewAnimationBlueprint(NewAnimBlueprint);
+}
+
+void FPreviewSceneDescriptionCustomization::HandleAnimBlueprintCompiled(UBlueprint* Blueprint)
+{
+	// Only re-initialize controller if we are not debugging an external instance.
+	// If we switch at this point then we will disconnect from the external instance
+	const TSharedPtr<FAnimationEditorPreviewScene> AnimPreviewScene = PreviewScene.Pin();
+	if(AnimPreviewScene->GetPreviewMeshComponent()->PreviewInstance == nullptr || AnimPreviewScene->GetPreviewMeshComponent()->PreviewInstance->GetDebugSkeletalMeshComponent() == nullptr)
+	{
+		UPersonaPreviewSceneDescription* PersonaPreviewSceneDescription = AnimPreviewScene->GetPreviewSceneDescription();
+		PersonaPreviewSceneDescription->PreviewControllerInstance->UninitializeView(PersonaPreviewSceneDescription, AnimPreviewScene.Get());
+		PersonaPreviewSceneDescription->PreviewControllerInstance->InitializeView(PersonaPreviewSceneDescription, AnimPreviewScene.Get());
+	}
+}
+
 void FPreviewSceneDescriptionCustomization::HandleAdditionalMeshesChanged(const FAssetData& InAssetData, IDetailLayoutBuilder* DetailLayoutBuilder)
 {
 	UDataAsset* MeshCollection = Cast<UDataAsset>(InAssetData.GetAsset());
diff --git a/Engine/Source/Editor/Persona/Private/PreviewSceneCustomizations.h b/Engine/Source/Editor/Persona/Private/PreviewSceneCustomizations.h
--- a/Engine/Source/Editor/Persona/Private/PreviewSceneCustomizations.h
+++ b/Engine/Source/Editor/Persona/Private/PreviewSceneCustomizations.h
@@ -70,6 +70,11 @@ private:
 
 	void HandleMeshChanged(const FAssetData& InAssetData);
 
+	void HandlePreviewAnimBlueprintChanged(const FAssetData& InAssetData);
+	
+	// Called when the anim blueprint being edited is compiled.
+	void HandleAnimBlueprintCompiled(UBlueprint* Blueprint);
+
 	void HandleAdditionalMeshesChanged(const FAssetData& InAssetData, IDetailLayoutBuilder* DetailLayoutBuilder);
 
 	void HandleAllowDifferentSkeletonsCheckedStateChanged(ECheckBoxState CheckState);
