Add missing nullptr check on animation blueprint editor preview instance

diff --git a/Engine/Source/Editor/AnimationBlueprintEditor/Private/AnimationBlueprintEditor.cpp b/Engine/Source/Editor/AnimationBlueprintEditor/Private/AnimationBlueprintEditor.cpp
--- a/Engine/Source/Editor/AnimationBlueprintEditor/Private/AnimationBlueprintEditor.cpp
+++ b/Engine/Source/Editor/AnimationBlueprintEditor/Private/AnimationBlueprintEditor.cpp
@@ -2311,7 +2311,10 @@ void FAnimationBlueprintEditor::HandleSetObjectBeingDebugged(UObject* InObject)
 			{
 				// Otherwise set us to display the debugged instance via copy-pose
 				GetPreviewScene()->GetPreviewMeshComponent()->EnablePreview(true, nullptr);
-				GetPreviewScene()->GetPreviewMeshComponent()->PreviewInstance->SetDebugSkeletalMeshComponent(SkeletalMeshComponent);
+				if (GetPreviewScene()->GetPreviewMeshComponent()->PreviewInstance)
+				{
+					GetPreviewScene()->GetPreviewMeshComponent()->PreviewInstance->SetDebugSkeletalMeshComponent(SkeletalMeshComponent);
+				}
 			}
 		}
 	}
