Fix notifies not being displayed in context menus under certain circumstances

Notifies only on skeletons only showed up when not in 'context sensitive' mode
State machine notifies were not showing up due to keying off the asset rather than the generated class

diff --git a/Engine/Source/Editor/AnimationBlueprintEditor/Private/AnimationBlueprintEditorModule.cpp b/Engine/Source/Editor/AnimationBlueprintEditor/Private/AnimationBlueprintEditorModule.cpp
--- a/Engine/Source/Editor/AnimationBlueprintEditor/Private/AnimationBlueprintEditorModule.cpp
+++ b/Engine/Source/Editor/AnimationBlueprintEditor/Private/AnimationBlueprintEditorModule.cpp
@@ -133,7 +133,15 @@ void FAnimationBlueprintEditorModule::GetTypeActions(FBlueprintActionDatabaseReg
 		const FString TagValue = AssetData.GetTagValueRef<FString>(USkeleton::AnimNotifyTag);
 		if (!TagValue.IsEmpty())
 		{
-			FSoftObjectPath SkeletonPath(AssetData.GetTagValueRef<FString>("Skeleton"));
+			FSoftObjectPath SkeletonPath;
+			if(AssetData.GetClass() == USkeleton::StaticClass())
+			{
+				SkeletonPath = AssetData.ToSoftObjectPath();
+			}
+			else
+			{
+				SkeletonPath = FSoftObjectPath(AssetData.GetTagValueRef<FString>("Skeleton"));
+			}
 			TSet<FName>& NotifyNames = NotifiesPerSkeleton.FindOrAdd(SkeletonPath);
 
 			UE::String::ParseTokens(TagValue, USkeleton::AnimNotifyTagDelimiter, [&NotifyNames](FStringView InToken)
@@ -173,7 +181,7 @@ void FAnimationBlueprintEditorModule::GetInstanceActions(const UAnimBlueprint* I
 			if (NotifyName != NAME_None)
 			{
 				UAnimNotifyEventNodeSpawner* NodeSpawner = UAnimNotifyEventNodeSpawner::Create(SkeletonPath, NotifyName);
-				ActionRegistrar.AddBlueprintAction(InAnimBlueprint, NodeSpawner);
+				ActionRegistrar.AddBlueprintAction(GeneratedClass, NodeSpawner);
 			}
 		}
 	}
