Update default LOD mode in Animation Blueprint preview to use the debug target's LOD

diff --git a/Engine/Source/Editor/AnimationBlueprintEditor/Private/AnimationBlueprintEditor.cpp b/Engine/Source/Editor/AnimationBlueprintEditor/Private/AnimationBlueprintEditor.cpp
--- a/Engine/Source/Editor/AnimationBlueprintEditor/Private/AnimationBlueprintEditor.cpp
+++ b/Engine/Source/Editor/AnimationBlueprintEditor/Private/AnimationBlueprintEditor.cpp
@@ -2305,7 +2305,11 @@ void FAnimationBlueprintEditor::HandleSetObjectBeingDebugged(UObject* InObject)
 				{
 					GetPreviewScene()->GetPreviewMeshComponent()->PreviewInstance->SetDebugSkeletalMeshComponent(nullptr);
 				}
-				GetPreviewScene()->GetPreviewMeshComponent()->bTrackAttachedInstanceLOD = false;
+				if (GetPreviewScene()->GetPreviewMeshComponent()->bTrackAttachedInstanceLOD)
+				{
+					GetPreviewScene()->GetPreviewMeshComponent()->bTrackAttachedInstanceLOD = false;
+					GetPreviewScene()->GetPreviewMeshComponent()->SetForcedLOD(0);
+				}
 			}
 			else
 			{
@@ -2314,6 +2318,7 @@ void FAnimationBlueprintEditor::HandleSetObjectBeingDebugged(UObject* InObject)
 				if (GetPreviewScene()->GetPreviewMeshComponent()->PreviewInstance)
 				{
 					GetPreviewScene()->GetPreviewMeshComponent()->PreviewInstance->SetDebugSkeletalMeshComponent(SkeletalMeshComponent);
+					GetPreviewScene()->GetPreviewMeshComponent()->bTrackAttachedInstanceLOD = true;
 				}
 			}
 		}
@@ -2326,7 +2331,12 @@ void FAnimationBlueprintEditor::HandleSetObjectBeingDebugged(UObject* InObject)
 		{
 			GetPreviewScene()->GetPreviewMeshComponent()->PreviewInstance->SetDebugSkeletalMeshComponent(nullptr);
 		}
-		GetPreviewScene()->GetPreviewMeshComponent()->bTrackAttachedInstanceLOD = false;
+		
+		if (GetPreviewScene()->GetPreviewMeshComponent()->bTrackAttachedInstanceLOD)
+		{
+			GetPreviewScene()->GetPreviewMeshComponent()->bTrackAttachedInstanceLOD = false;
+			GetPreviewScene()->GetPreviewMeshComponent()->SetForcedLOD(0);
+		}
 	}
 }
 
