Add missing fast path icons for blend space nodes
* Functionality copied over from SAnimationGraphNode

diff --git a/Engine/Source/Editor/AnimationBlueprintEditor/Private/AnimationNodes/SGraphNodeBlendSpaceGraph.cpp b/Engine/Source/Editor/AnimationBlueprintEditor/Private/AnimationNodes/SGraphNodeBlendSpaceGraph.cpp
--- a/Engine/Source/Editor/AnimationBlueprintEditor/Private/AnimationNodes/SGraphNodeBlendSpaceGraph.cpp
+++ b/Engine/Source/Editor/AnimationBlueprintEditor/Private/AnimationNodes/SGraphNodeBlendSpaceGraph.cpp
@@ -47,6 +47,7 @@
 #include "Widgets/SBoxPanel.h"
 #include "Widgets/SToolTip.h"
 #include "Widgets/Text/STextBlock.h"
+#include "Widgets/Images/SImage.h"
 
 class SWidget;
 class UEdGraphPin;
@@ -59,6 +60,11 @@ void SGraphNodeBlendSpaceGraph::Construct(const FArguments& InArgs, UAnimGraphNo
 
 	SetCursor(EMouseCursor::CardinalCross);
 
+	FastPathWidget = SNew(SImage)
+		.Image(FAppStyle::Get().GetBrush(TEXT("Graph.AnimationFastPathIndicator")))
+		.ToolTip(IDocumentation::Get()->CreateToolTip(LOCTEXT("AnimGraphNodeIndicatorTooltip", "Fast path enabled: This node is not using any Blueprint calls to update its data."), nullptr, TEXT("Shared/GraphNodes/Animation"), TEXT("GraphNode_FastPathInfo")))
+		.Visibility(EVisibility::Visible);
+
 	PoseWatchWidget = SNew(SPoseWatchOverlay, InNode);
 
 	UpdateGraphNode();
@@ -79,6 +85,17 @@ TArray<FOverlayWidgetInfo> SGraphNodeBlendSpaceGraph::GetOverlayWidgets(bool bSe
 
 	if (UAnimGraphNode_Base* AnimNode = CastChecked<UAnimGraphNode_Base>(GraphNode, ECastCheckedType::NullAllowed))
 	{
+		if (AnimNode->BlueprintUsage == EBlueprintUsage::DoesNotUseBlueprint)
+		{
+			const FSlateBrush* ImageBrush = FAppStyle::Get().GetBrush(TEXT("Graph.AnimationFastPathIndicator"));
+
+			FOverlayWidgetInfo Info;
+			Info.OverlayOffset = FVector2D(WidgetSize.X - (ImageBrush->ImageSize.X * 0.5f), -(ImageBrush->ImageSize.Y * 0.5f));
+			Info.Widget = FastPathWidget;
+
+			Widgets.Add(Info);
+		}
+
 		if (PoseWatchWidget->IsPoseWatchValid())
 		{
 			FOverlayWidgetInfo Info;
diff --git a/Engine/Source/Editor/AnimationBlueprintEditor/Private/AnimationNodes/SGraphNodeBlendSpaceGraph.h b/Engine/Source/Editor/AnimationBlueprintEditor/Private/AnimationNodes/SGraphNodeBlendSpaceGraph.h
--- a/Engine/Source/Editor/AnimationBlueprintEditor/Private/AnimationNodes/SGraphNodeBlendSpaceGraph.h
+++ b/Engine/Source/Editor/AnimationBlueprintEditor/Private/AnimationNodes/SGraphNodeBlendSpaceGraph.h
@@ -32,5 +32,7 @@ protected:
 	TSharedPtr<SToolTip> GetComplexTooltip() override;
 
 private:
+	TSharedPtr<SWidget> FastPathWidget;
+
 	TSharedPtr<SPoseWatchOverlay> PoseWatchWidget;
 };
