Add property to disable transition in state machine editor
* Also update coloring of shared transition rules to have full transparency and reserve translucent transitions for representing disabled transitions

diff --git a/Engine/Source/Editor/AnimGraph/Private/AnimGraphNode_StateMachineBase.cpp b/Engine/Source/Editor/AnimGraph/Private/AnimGraphNode_StateMachineBase.cpp
--- a/Engine/Source/Editor/AnimGraph/Private/AnimGraphNode_StateMachineBase.cpp
+++ b/Engine/Source/Editor/AnimGraph/Private/AnimGraphNode_StateMachineBase.cpp
@@ -614,6 +614,11 @@ void UAnimGraphNode_StateMachineBase::OnProcessDuringCompilation(IAnimBlueprintC
 		{
 			UAnimStateTransitionNode* TransitionNode = *TransitionIt;
 
+			if (TransitionNode->bDisabled)
+			{
+				continue;
+			}
+			
 			const int32 TransitionIndex = Oven.FindOrAddTransition(FMachineCreator::FUniqueTransition(TransitionNode, StateNode));
 			FAnimationTransitionBetweenStates& BakedTransition = Oven.GetMachine().Transitions[TransitionIndex];
 
diff --git a/Engine/Source/Editor/AnimGraph/Private/AnimStateTransitionNode.cpp b/Engine/Source/Editor/AnimGraph/Private/AnimStateTransitionNode.cpp
--- a/Engine/Source/Editor/AnimGraph/Private/AnimStateTransitionNode.cpp
+++ b/Engine/Source/Editor/AnimGraph/Private/AnimStateTransitionNode.cpp
@@ -87,6 +87,7 @@ UAnimStateTransitionNode::UAnimStateTransitionNode(const FObjectInitializer& Obj
 	SharedCrossfadeIdx = INDEX_NONE;
 	SharedCrossfadeGuid.Invalidate();
 	Bidirectional = false;
+	bDisabled = false;
 	PriorityOrder = 1;
 	LogicType = ETransitionLogicType::TLT_StandardBlend;
 }
diff --git a/Engine/Source/Editor/AnimGraph/Public/AnimStateTransitionNode.h b/Engine/Source/Editor/AnimGraph/Public/AnimStateTransitionNode.h
--- a/Engine/Source/Editor/AnimGraph/Public/AnimStateTransitionNode.h
+++ b/Engine/Source/Editor/AnimGraph/Public/AnimStateTransitionNode.h
@@ -83,6 +83,10 @@ class UAnimStateTransitionNode : public UAnimStateNodeBase
 	UPROPERTY(EditAnywhere, Category=Transition)
 	bool Bidirectional;
 
+	/** Whether or not to disable this transition from being enterable */
+	UPROPERTY(EditAnywhere, Category=Transition)
+	bool bDisabled;
+
 	/** The rules for this transition may be shared with other transition nodes */
 	UPROPERTY()
 	bool bSharedRules;
diff --git a/Engine/Source/Editor/AnimationBlueprintEditor/Private/AnimationStateNodes/SGraphNodeAnimTransition.cpp b/Engine/Source/Editor/AnimationBlueprintEditor/Private/AnimationStateNodes/SGraphNodeAnimTransition.cpp
--- a/Engine/Source/Editor/AnimationBlueprintEditor/Private/AnimationStateNodes/SGraphNodeAnimTransition.cpp
+++ b/Engine/Source/Editor/AnimationBlueprintEditor/Private/AnimationStateNodes/SGraphNodeAnimTransition.cpp
@@ -201,6 +201,17 @@ TSharedRef<SWidget> SGraphNodeAnimTransition::GenerateRichTooltip()
 			];
 	}
 
+	if (TransNode->bDisabled)
+	{
+		Widget->AddSlot()
+			.AutoHeight()
+			.Padding(2.0f)
+			[
+				SNew(STextBlock)
+				.Text(LOCTEXT("TransitionRuleDisabled", "Disabled"))
+			];
+	}
+
 	Widget->AddSlot()
 		.AutoHeight()
 		.Padding( 2.0f )
@@ -353,6 +364,11 @@ FLinearColor SGraphNodeAnimTransition::StaticGetTransitionColor(UAnimStateTransi
 		BaseColor = TransNode->SharedColor;
 	}
 
+	if (TransNode->bDisabled)
+	{
+		BaseColor = FLinearColor(0.125, 0.125, 0.125, 0.5);
+	}
+
 	return bIsHovered ? HoverColor : BaseColor;
 }
 
diff --git a/Engine/Source/Editor/DetailCustomizations/Private/AnimTransitionNodeDetails.cpp b/Engine/Source/Editor/DetailCustomizations/Private/AnimTransitionNodeDetails.cpp
--- a/Engine/Source/Editor/DetailCustomizations/Private/AnimTransitionNodeDetails.cpp
+++ b/Engine/Source/Editor/DetailCustomizations/Private/AnimTransitionNodeDetails.cpp
@@ -477,7 +477,7 @@ void FAnimTransitionNodeDetails::AssignUniqueColorsToAllSharedNodes(UEdGraph* Cu
 				SharedColor.R = (colorIdx & 1 ? 1.0f : 0.15f);
 				SharedColor.G = (colorIdx & 2 ? 1.0f : 0.15f);
 				SharedColor.B = (colorIdx & 4 ? 1.0f : 0.15f);
-				SharedColor.A = 0.25f;
+				SharedColor.A = 1.0f;
 
 				// Storing this on the UAnimStateTransitionNode really bugs me. But its a pain to iterate over all the widget nodes at once
 				// and we may want the shared color to be customizable in the details view
