AnimationBlueprintEditor - fix for FAnimationBlueprintEditor sometimes referencing a GCed DebuggedMeshComponent

diff --git a/Engine/Source/Editor/AnimationBlueprintEditor/Private/AnimationBlueprintEditor.h b/Engine/Source/Editor/AnimationBlueprintEditor/Private/AnimationBlueprintEditor.h
--- a/Engine/Source/Editor/AnimationBlueprintEditor/Private/AnimationBlueprintEditor.h
+++ b/Engine/Source/Editor/AnimationBlueprintEditor/Private/AnimationBlueprintEditor.h
@@ -414,7 +414,7 @@ private:
 	TObjectPtr<UAnimationBlueprintEditorOptions> EditorOptions;
 
 	/** Cached mesh component held during compilation, used to reconnect debugger */
-	USkeletalMeshComponent* DebuggedMeshComponent;
+	TWeakObjectPtr<USkeletalMeshComponent> DebuggedMeshComponent;
 
 	/** Used to track wither the editor option has changed */
 	bool bPreviousPoseWatchSelectedNodes = false;
