Fixed longer (3-4x) animbp execution times in PIE if you had recently opened and closed the corresponding animbp in editor.

FAnimationBlueprintEditorMode wasn't clearing the animbp's ObjectBeingDebugged on close. This left a stale preview actor as ObjectPathToDebug on the blueprint, which caused the bp debugger to try searching for this stale object on each node execution.

diff --git a/Engine/Source/Editor/AnimationBlueprintEditor/Private/AnimationBlueprintEditorMode.cpp b/Engine/Source/Editor/AnimationBlueprintEditor/Private/AnimationBlueprintEditorMode.cpp
--- a/Engine/Source/Editor/AnimationBlueprintEditor/Private/AnimationBlueprintEditorMode.cpp
+++ b/Engine/Source/Editor/AnimationBlueprintEditor/Private/AnimationBlueprintEditorMode.cpp
@@ -217,6 +217,15 @@ void FAnimationBlueprintEditorMode::RegisterTabFactories(TSharedPtr<FTabManager>
 	BP->PushTabFactories(TabFactories);
 }
 
+void FAnimationBlueprintEditorMode::PreDeactivateMode()
+{
+	if (UAnimBlueprint* AnimBlueprint = AnimBlueprintPtr.Get())
+	{
+		AnimBlueprint->SetObjectBeingDebugged(nullptr);
+	}
+	FBlueprintEditorApplicationMode::PreDeactivateMode();
+}
+
 void FAnimationBlueprintEditorMode::PostActivateMode()
 {
 	if (UAnimBlueprint* AnimBlueprint = AnimBlueprintPtr.Get())
diff --git a/Engine/Source/Editor/AnimationBlueprintEditor/Private/AnimationBlueprintEditorMode.h b/Engine/Source/Editor/AnimationBlueprintEditor/Private/AnimationBlueprintEditorMode.h
--- a/Engine/Source/Editor/AnimationBlueprintEditor/Private/AnimationBlueprintEditorMode.h
+++ b/Engine/Source/Editor/AnimationBlueprintEditor/Private/AnimationBlueprintEditorMode.h
@@ -24,6 +24,7 @@ public:
 
 	// FApplicationMode interface
 	virtual void RegisterTabFactories(TSharedPtr<FTabManager> InTabManager) override;
+	virtual void PreDeactivateMode() override;
 	virtual void PostActivateMode() override;
 	// End of FApplicationMode interface
 
