Add support for following tokenized hyperlinks in animation blueprint logging

diff --git a/Engine/Source/Editor/AnimationBlueprintEditor/Private/AnimationBlueprintEditorModule.cpp b/Engine/Source/Editor/AnimationBlueprintEditor/Private/AnimationBlueprintEditorModule.cpp
--- a/Engine/Source/Editor/AnimationBlueprintEditor/Private/AnimationBlueprintEditorModule.cpp
+++ b/Engine/Source/Editor/AnimationBlueprintEditor/Private/AnimationBlueprintEditorModule.cpp
@@ -37,6 +37,8 @@
 #include "Animation/Skeleton.h"
 #include "Animation/AnimBlueprint.h"
 #include "Animation/AnimBlueprintGeneratedClass.h"
+#include "IMessageLogListing.h"
+#include "Misc/UObjectToken.h"
 
 class IAnimationBlueprintEditor;
 class IToolkitHost;
@@ -66,6 +68,9 @@ void FAnimationBlueprintEditorModule::StartupModule()
 	InitOptions.bShowFilters = true;
 	InitOptions.bShowPages = true;
 	MessageLogModule.RegisterLogListing("AnimBlueprintLog", LOCTEXT("AnimBlueprintLog", "Anim Blueprint Log"), InitOptions);
+	
+	TSharedRef<IMessageLogListing> LogListing = MessageLogModule.GetLogListing("AnimBlueprintLog");
+	LogListing->OnMessageTokenClicked().AddRaw(this, &FAnimationBlueprintEditorModule::HandleMessageTokenClicked);
 
 	ISettingsModule& SettingsModule = FModuleManager::LoadModuleChecked<ISettingsModule>("Settings");
 	SettingsModule.RegisterSettings("Editor", "ContentEditors", "AnimationBlueprintEditor",
@@ -83,6 +88,10 @@ void FAnimationBlueprintEditorModule::ShutdownModule()
 	FEdGraphUtilities::UnregisterVisualPinFactory(AnimGraphPinFactory);
 	FEdGraphUtilities::UnregisterVisualPinConnectionFactory(AnimGraphPinConnectionFactory);
 
+	FMessageLogModule& MessageLogModule = FModuleManager::LoadModuleChecked<FMessageLogModule>("MessageLog");
+	TSharedRef<IMessageLogListing> LogListing = MessageLogModule.GetLogListing("AnimBlueprintLog");
+	LogListing->OnMessageTokenClicked().RemoveAll(this);
+
 	ISettingsModule& SettingsModule = FModuleManager::LoadModuleChecked<ISettingsModule>("Settings");
 	SettingsModule.UnregisterSettings("Editor", "ContentEditors", "AnimationBlueprintEditor");
 
@@ -219,4 +228,58 @@ void FAnimationBlueprintEditorModule::OnNewBlueprintCreated(UBlueprint* InBluepr
 	}
 }
 
+void FAnimationBlueprintEditorModule::HandleMessageTokenClicked(const TSharedRef<IMessageToken>& InMessageToken)
+{
+	if (InMessageToken->GetType() == EMessageToken::Object)
+	{
+		const TSharedRef<FUObjectToken> ObjectToken = StaticCastSharedRef<FUObjectToken>(InMessageToken);
+		UEdGraphNode* Node = Cast<UEdGraphNode>(ObjectToken->GetObject().Get());
+		if (Node)
+		{
+			UEdGraph* Graph = Node->GetGraph();
+			if (UAnimBlueprint* AnimBlueprint = Cast<UAnimBlueprint>(Graph->GetOuter()))
+			{
+				GEditor->GetEditorSubsystem<UAssetEditorSubsystem>()->OpenEditorForAsset(AnimBlueprint);
+				if (IAnimationBlueprintEditor* AnimBlueprintEditor =
+					static_cast<IAnimationBlueprintEditor*>(GEditor->GetEditorSubsystem<UAssetEditorSubsystem>()->FindEditorForAsset(AnimBlueprint, true)))
+				{
+					AnimBlueprintEditor->JumpToHyperlink(Node);
+				}
+			}
+		}
+	}
+	else if (InMessageToken->GetType() == EMessageToken::EdGraph)
+	{
+		const TSharedRef<FEdGraphToken> EdGraphToken = StaticCastSharedRef<FEdGraphToken>(InMessageToken);
+		const UEdGraphPin* Pin = EdGraphToken->GetPin();
+		const UEdGraphNode* Node = Cast<UEdGraphNode>(EdGraphToken->GetGraphObject());
+		if (Pin)
+		{
+			UEdGraph* Graph = Pin->GetOwningNode()->GetGraph();
+			if (UAnimBlueprint* AnimBlueprint = Cast<UAnimBlueprint>(Graph->GetOuter()))
+			{
+				GEditor->GetEditorSubsystem<UAssetEditorSubsystem>()->OpenEditorForAsset(AnimBlueprint);
+				if (IAnimationBlueprintEditor* AnimBlueprintEditor =
+					static_cast<IAnimationBlueprintEditor*>(GEditor->GetEditorSubsystem<UAssetEditorSubsystem>()->FindEditorForAsset(AnimBlueprint, true)))
+				{
+					AnimBlueprintEditor->JumpToPin(Pin);
+				}
+			}
+		}
+		else if (Node)
+		{
+			UEdGraph* Graph = Node->GetGraph();
+			if (UAnimBlueprint* AnimBlueprint = Cast<UAnimBlueprint>(Graph->GetOuter()))
+			{
+				GEditor->GetEditorSubsystem<UAssetEditorSubsystem>()->OpenEditorForAsset(AnimBlueprint);
+				if (IAnimationBlueprintEditor* AnimBlueprintEditor =
+					static_cast<IAnimationBlueprintEditor*>(GEditor->GetEditorSubsystem<UAssetEditorSubsystem>()->FindEditorForAsset(AnimBlueprint, true)))
+				{
+					AnimBlueprintEditor->JumpToHyperlink(Node);
+				}
+			}
+		}
+	}
+}
+
 #undef LOCTEXT_NAMESPACE
diff --git a/Engine/Source/Editor/AnimationBlueprintEditor/Private/AnimationBlueprintEditorModule.h b/Engine/Source/Editor/AnimationBlueprintEditor/Private/AnimationBlueprintEditorModule.h
--- a/Engine/Source/Editor/AnimationBlueprintEditor/Private/AnimationBlueprintEditorModule.h
+++ b/Engine/Source/Editor/AnimationBlueprintEditor/Private/AnimationBlueprintEditorModule.h
@@ -43,6 +43,7 @@ public:
 private:
 	/** When a new AnimBlueprint is created, this will handle post creation work such as adding non-event default nodes */
 	void OnNewBlueprintCreated(class UBlueprint* InBlueprint);
+	void HandleMessageTokenClicked(const TSharedRef<IMessageToken>& InMessageToken);
 
 private:
 	TSharedPtr<class FExtensibilityManager> MenuExtensibilityManager;
