Fix crash opening abp

#rnx

diff --git a/Engine/Source/Editor/AnimationBlueprintEditor/Private/AnimationBlueprintEditorMode.cpp b/Engine/Source/Editor/AnimationBlueprintEditor/Private/AnimationBlueprintEditorMode.cpp
--- a/Engine/Source/Editor/AnimationBlueprintEditor/Private/AnimationBlueprintEditorMode.cpp
+++ b/Engine/Source/Editor/AnimationBlueprintEditor/Private/AnimationBlueprintEditorMode.cpp
@@ -178,7 +178,10 @@ FAnimationBlueprintEditorMode::FAnimationBlueprintEditorMode(const TSharedRef<FA
 	TabFactories.RegisterFactory(PersonaModule.CreatePoseWatchTabFactory(InAnimationBlueprintEditor));
 	TSharedPtr<ISkeletonTree> SkeletonTree = InAnimationBlueprintEditor->GetSkeletonTree();
 	TabFactories.RegisterFactory(PersonaModule.CreateCurveViewerTabFactory(InAnimationBlueprintEditor, SkeletonTree.IsValid() ? SkeletonTree->GetEditableSkeleton().ToSharedPtr() : nullptr, InAnimationBlueprintEditor->GetPersonaToolkit()->GetPreviewScene(), FOnObjectsSelected::CreateSP(&InAnimationBlueprintEditor.Get(), &FAnimationBlueprintEditor::HandleObjectsSelected)));
-	TabFactories.RegisterFactory(PersonaModule.CreateAttributeViewerTabFactory(InAnimationBlueprintEditor, InAnimationBlueprintEditor->GetSkeletonTree()->GetEditableSkeleton(), InAnimationBlueprintEditor->GetPersonaToolkit()->GetPreviewScene()));	
+	if (SkeletonTree.IsValid())
+	{
+		TabFactories.RegisterFactory(PersonaModule.CreateAttributeViewerTabFactory(InAnimationBlueprintEditor, SkeletonTree->GetEditableSkeleton(), InAnimationBlueprintEditor->GetPersonaToolkit()->GetPreviewScene()));
+	}
 	TabFactories.RegisterFactory(PersonaModule.CreateAnimAssetFindReplaceTabFactory(InAnimationBlueprintEditor, FAnimAssetFindReplaceConfig()));
 
 	if (!bIsTemplate)
