Removed unnecessary skeleton tree reference from AttributeViewer tab-spawner

diff --git a/Engine/Source/Editor/AnimationBlueprintEditor/Private/AnimationBlueprintEditorMode.cpp b/Engine/Source/Editor/AnimationBlueprintEditor/Private/AnimationBlueprintEditorMode.cpp
--- a/Engine/Source/Editor/AnimationBlueprintEditor/Private/AnimationBlueprintEditorMode.cpp
+++ b/Engine/Source/Editor/AnimationBlueprintEditor/Private/AnimationBlueprintEditorMode.cpp
@@ -178,10 +178,7 @@ FAnimationBlueprintEditorMode::FAnimationBlueprintEditorMode(const TSharedRef<FA
 	TabFactories.RegisterFactory(PersonaModule.CreatePoseWatchTabFactory(InAnimationBlueprintEditor));
 	TSharedPtr<ISkeletonTree> SkeletonTree = InAnimationBlueprintEditor->GetSkeletonTree();
 	TabFactories.RegisterFactory(PersonaModule.CreateCurveViewerTabFactory(InAnimationBlueprintEditor, SkeletonTree.IsValid() ? SkeletonTree->GetEditableSkeleton().ToSharedPtr() : nullptr, InAnimationBlueprintEditor->GetPersonaToolkit()->GetPreviewScene(), FOnObjectsSelected::CreateSP(&InAnimationBlueprintEditor.Get(), &FAnimationBlueprintEditor::HandleObjectsSelected)));
-	if (SkeletonTree.IsValid())
-	{
-		TabFactories.RegisterFactory(PersonaModule.CreateAttributeViewerTabFactory(InAnimationBlueprintEditor, SkeletonTree->GetEditableSkeleton(), InAnimationBlueprintEditor->GetPersonaToolkit()->GetPreviewScene()));
-	}
+	TabFactories.RegisterFactory(PersonaModule.CreateAttributeViewerTabFactory(InAnimationBlueprintEditor, InAnimationBlueprintEditor->GetPersonaToolkit()->GetPreviewScene()));
 	TabFactories.RegisterFactory(PersonaModule.CreateAnimAssetFindReplaceTabFactory(InAnimationBlueprintEditor, FAnimAssetFindReplaceConfig()));
 
 	if (!bIsTemplate)
diff --git a/Engine/Source/Editor/AnimationEditor/Private/AnimationEditorMode.cpp b/Engine/Source/Editor/AnimationEditor/Private/AnimationEditorMode.cpp
--- a/Engine/Source/Editor/AnimationEditor/Private/AnimationEditorMode.cpp
+++ b/Engine/Source/Editor/AnimationEditor/Private/AnimationEditorMode.cpp
@@ -46,7 +46,7 @@ FAnimationEditorMode::FAnimationEditorMode(TSharedRef<FWorkflowCentricApplicatio
 	TabFactories.RegisterFactory(PersonaModule.CreateAnimNotifiesTabFactory(InHostingApp, InSkeletonTree->GetEditableSkeleton(), OnObjectsSelected));
 	TabFactories.RegisterFactory(PersonaModule.CreateAnimMontageSectionsTabFactory(InHostingApp, AnimationEditor->GetPersonaToolkit(), AnimationEditor->OnSectionsChanged));
 	TabFactories.RegisterFactory(PersonaModule.CreateAnimAssetFindReplaceTabFactory(InHostingApp, FAnimAssetFindReplaceConfig()));
-	TabFactories.RegisterFactory(PersonaModule.CreateAttributeViewerTabFactory(InHostingApp, AnimationEditor->GetSkeletonTree()->GetEditableSkeleton(), AnimationEditor->GetPersonaToolkit()->GetPreviewScene()));
+	TabFactories.RegisterFactory(PersonaModule.CreateAttributeViewerTabFactory(InHostingApp, AnimationEditor->GetPersonaToolkit()->GetPreviewScene()));
 
 	TabLayout = FTabManager::NewLayout("Standalone_AnimationEditor_Layout_v1.5")
 		->AddArea
diff --git a/Engine/Source/Editor/Persona/Private/PersonaModule.cpp b/Engine/Source/Editor/Persona/Private/PersonaModule.cpp
--- a/Engine/Source/Editor/Persona/Private/PersonaModule.cpp
+++ b/Engine/Source/Editor/Persona/Private/PersonaModule.cpp
@@ -381,10 +381,9 @@ TSharedRef<FWorkflowTabFactory> FPersonaModule::CreatePersonaToolboxTabFactory(c
 	return MakeShareable(new FToolBoxSummoner(InHostingApp));
 }
 
-TSharedRef<FWorkflowTabFactory> FPersonaModule::CreateAttributeViewerTabFactory(const TSharedRef<FWorkflowCentricApplication>& InHostingApp, const TSharedRef<IEditableSkeleton>& InEditableSkeleton,
-	const TSharedRef<IPersonaPreviewScene>& InPreviewScene) const
+TSharedRef<FWorkflowTabFactory> FPersonaModule::CreateAttributeViewerTabFactory(const TSharedRef<FWorkflowCentricApplication>& InHostingApp, const TSharedRef<IPersonaPreviewScene>& InPreviewScene) const
 {
-	return MakeShareable(new FAnimAttributeViewerTabSummoner(InHostingApp, InEditableSkeleton, InPreviewScene));
+	return MakeShareable(new FAnimAttributeViewerTabSummoner(InHostingApp, InPreviewScene));
 }
 
 TSharedRef<SWidget> FPersonaModule::CreateBlendSpacePreviewWidget(TAttribute<const UBlendSpace*> InBlendSpace, TAttribute<FVector> InPosition, TAttribute<FVector> InFilteredPosition) const
diff --git a/Engine/Source/Editor/Persona/Private/SAnimAttributeView.cpp b/Engine/Source/Editor/Persona/Private/SAnimAttributeView.cpp
--- a/Engine/Source/Editor/Persona/Private/SAnimAttributeView.cpp
+++ b/Engine/Source/Editor/Persona/Private/SAnimAttributeView.cpp
@@ -511,10 +511,9 @@ SAnimAttributeView::FAttributeValueView::FAttributeValueView(FName InSnapshotNam
 	ViewWidget->SetStructureData(StructData);
 }
 
-void SAnimAttributeViewer::Construct(const FArguments& InArgs, const TSharedRef<IEditableSkeleton>& InEditableSkeleton, const TSharedRef<IPersonaPreviewScene>& InPreviewScene)
+void SAnimAttributeViewer::Construct(const FArguments& InArgs, const TSharedRef<IPersonaPreviewScene>& InPreviewScene)
 {
 	PreviewScenePtr = InPreviewScene;
-	EditableSkeletonPtr = InEditableSkeleton;
 
 	SAssignNew(AttributeView, SAnimAttributeView);
 	SAssignNew(PoseWatchPicker, SPoseWatchPicker)
diff --git a/Engine/Source/Editor/Persona/Private/TabSpawners.cpp b/Engine/Source/Editor/Persona/Private/TabSpawners.cpp
--- a/Engine/Source/Editor/Persona/Private/TabSpawners.cpp
+++ b/Engine/Source/Editor/Persona/Private/TabSpawners.cpp
@@ -836,9 +836,8 @@ TSharedRef<SWidget> FAssetPropertiesSummoner::CreateTabBody(const FWorkflowTabSp
 /////////////////////////////////////////////////////
 // FAnimAttributeViewerTabSummoner
 
-FAnimAttributeViewerTabSummoner::FAnimAttributeViewerTabSummoner(TSharedPtr<class FAssetEditorToolkit> InHostingApp, const TSharedRef<IEditableSkeleton>& InEditableSkeleton, const TSharedRef<IPersonaPreviewScene>& InPreviewScene)
+FAnimAttributeViewerTabSummoner::FAnimAttributeViewerTabSummoner(TSharedPtr<class FAssetEditorToolkit> InHostingApp, const TSharedRef<IPersonaPreviewScene>& InPreviewScene)
 	: FWorkflowTabFactory(FPersonaTabs::AnimAttributeViewID, InHostingApp)
-	, EditableSkeleton(InEditableSkeleton)
 	, PreviewScene(InPreviewScene)
 {
 	TabLabel = LOCTEXT("AnimAttributeViewTabTitle", "Animation Attributes");	
@@ -853,7 +852,7 @@ FAnimAttributeViewerTabSummoner::FAnimAttributeViewerTabSummoner(TSharedPtr<clas
 
 TSharedRef<SWidget> FAnimAttributeViewerTabSummoner::CreateTabBody(const FWorkflowTabSpawnInfo& Info) const
 {
-	return SNew(SAnimAttributeViewer, EditableSkeleton.Pin().ToSharedRef(), PreviewScene.Pin().ToSharedRef());
+	return SNew(SAnimAttributeViewer, PreviewScene.Pin().ToSharedRef());
 }
 
 
diff --git a/Engine/Source/Editor/Persona/Private/TabSpawners.h b/Engine/Source/Editor/Persona/Private/TabSpawners.h
--- a/Engine/Source/Editor/Persona/Private/TabSpawners.h
+++ b/Engine/Source/Editor/Persona/Private/TabSpawners.h
@@ -415,7 +415,7 @@ private:
 struct FAnimAttributeViewerTabSummoner : public FWorkflowTabFactory
 {
 public:
-	FAnimAttributeViewerTabSummoner(TSharedPtr<class FAssetEditorToolkit> InHostingApp, const TSharedRef<IEditableSkeleton>& InEditableSkeleton, const TSharedRef<IPersonaPreviewScene>& InPreviewScene);
+	FAnimAttributeViewerTabSummoner(TSharedPtr<class FAssetEditorToolkit> InHostingApp, const TSharedRef<IPersonaPreviewScene>& InPreviewScene);
 
 	virtual TSharedRef<SWidget> CreateTabBody(const FWorkflowTabSpawnInfo& Info) const override;
 
diff --git a/Engine/Source/Editor/Persona/Public/PersonaModule.h b/Engine/Source/Editor/Persona/Public/PersonaModule.h
--- a/Engine/Source/Editor/Persona/Public/PersonaModule.h
+++ b/Engine/Source/Editor/Persona/Public/PersonaModule.h
@@ -410,7 +410,7 @@ public:
 	virtual TSharedRef<FWorkflowTabFactory> CreatePersonaToolboxTabFactory(const TSharedRef<class FPersonaAssetEditorToolkit>& InHostingApp) const;
 
 	/** Create a attributes viewer tab factory */
-	virtual TSharedRef<FWorkflowTabFactory> CreateAttributeViewerTabFactory(const TSharedRef<class FWorkflowCentricApplication>& InHostingApp, const TSharedRef<class IEditableSkeleton>& InEditableSkeleton, const TSharedRef<class IPersonaPreviewScene>& InPreviewScene) const;
+	virtual TSharedRef<FWorkflowTabFactory> CreateAttributeViewerTabFactory(const TSharedRef<class FWorkflowCentricApplication>& InHostingApp, const TSharedRef<class IPersonaPreviewScene>& InPreviewScene) const;
 
 	/** Deprecated */
 	UE_DEPRECATED(5.0, "Please use the overload that takes a FBlendSpacePreviewArgs struct")
diff --git a/Engine/Source/Editor/Persona/Public/SAnimAttributeView.h b/Engine/Source/Editor/Persona/Public/SAnimAttributeView.h
--- a/Engine/Source/Editor/Persona/Public/SAnimAttributeView.h
+++ b/Engine/Source/Editor/Persona/Public/SAnimAttributeView.h
@@ -379,26 +379,19 @@ inline void SAnimAttributeView::ClearListView()
 
 class PERSONA_API SAnimAttributeViewer : public SCompoundWidget
 {
-	SLATE_BEGIN_ARGS( SAnimAttributeViewer )
-	{}
-	
+	SLATE_BEGIN_ARGS( SAnimAttributeViewer ) {}
 	SLATE_END_ARGS()
 	
-	void Construct( const FArguments& InArgs, const TSharedRef<class IEditableSkeleton>& InEditableSkeleton, const TSharedRef<class IPersonaPreviewScene>& InPreviewScene);
-	
-	virtual ~SAnimAttributeViewer() override {}
-	
+	void Construct( const FArguments& InArgs, const TSharedRef<class IPersonaPreviewScene>& InPreviewScene);	
+	virtual ~SAnimAttributeViewer() override {}	
 	virtual void Tick(const FGeometry& AllottedGeometry, const double InCurrentTime, const float InDeltaTime) override;
 
+protected:
 	UAnimInstance* GetAnimInstance() const;
-
 protected:
 	/** Pointer to the preview scene we are bound to */
 	TWeakPtr<class IPersonaPreviewScene> PreviewScenePtr;
 
-	/** Pointer to the editable skeleton */
-	TWeakPtr<class IEditableSkeleton> EditableSkeletonPtr;
-
 	TSharedPtr<SAnimAttributeView> AttributeView;	
 	TSharedPtr<SPoseWatchPicker> PoseWatchPicker;
 };
