State Machine: Self transitions & UX improvements

Self transitions and state re-entry are now possible. State re-entry can be enabled using the new bAllowInertializationForSelfTransitions flag on transitions, and will use inertialization/dead-blending when re-entering a state that currently has weight.
To aid with managing state-reentry, a new MinTimeBeforeReentry value has been added to transitions. This prevents the transition from being taken if the state was last entered within the provided time limit.
Added 'create self transition' context menu shortcut that will create a state alias that is aliased only to the single desired state, along with a transition back to the state. Self transitions are set up with bAllowInertializationForSelfTransitions = true and MinTimeBeforeReentry = 0.0.
Added 'reverse transition' context menu option that just reverses the direction of transitions.
Relinking transitions to other states can now be performed from the transition tail as well as the head.
State alias list is now sorted alphabetically to make finding states easier

Look and feel of state machines has been tweaked & improved. Now uses SVG brushes to better scale.
- Different node types (states, aliases and conduits) now have different shapes
- Transitions now rotate to follow their direction and have been moved closer together and closer to their link line
- Transitions now sort in front of states
- Links are now further apart and scale correctly with zoom level
- Selection of links now scales correctly with zoom level

Details panels for states now show the anim node functions for their result node and have been rearranged.
Added 'force transition' option to the context menu for state machines that are currently being debugged.
Made machine entry nodes non-deletable.

Implementation notes:
- Because of the need to reference more pins via pin handles, input pins for states are now exposed as widgets (but are not interactable).
- Styling has been kept seperate via new keys, as styling is re-used elsewhere
- Updated FAnimInstanceProxy APIs to correctly refer to node index rather than state index.
- Fixed crash when copy-pasting transitions with getter functions.

diff --git a/Engine/Source/Editor/AnimGraph/Private/AnimGraphCommands.cpp b/Engine/Source/Editor/AnimGraph/Private/AnimGraphCommands.cpp
--- a/Engine/Source/Editor/AnimGraph/Private/AnimGraphCommands.cpp
+++ b/Engine/Source/Editor/AnimGraph/Private/AnimGraphCommands.cpp
@@ -33,6 +33,10 @@ void FAnimGraphCommands::RegisterCommands()
 	UI_COMMAND(ConvertToPoseByName, "Convert to Pose By Name", "Convert to pose node that returns by name", EUserInterfaceActionType::Button, FInputChord())
 
 	UI_COMMAND( OpenRelatedAsset, "Open Asset", "Opens the asset related to this node", EUserInterfaceActionType::Button, FInputChord() )
+
+	UI_COMMAND( CreateSelfTransition, "Create Self Transition", "Creates an alias for this state and adds a 'self' transition from the alias to this state", EUserInterfaceActionType::Button, FInputChord() )
+	UI_COMMAND( SetActiveState, "Set Active State", "Sets the currently-selected state as active", EUserInterfaceActionType::Button, FInputChord() )
+	UI_COMMAND( ReverseTransition, "Reverse Transition", "Reverses this transition", EUserInterfaceActionType::Button, FInputChord() )
 }
 
 #undef LOCTEXT_NAMESPACE
diff --git a/Engine/Source/Editor/AnimGraph/Private/AnimGraphNode_StateMachineBase.cpp b/Engine/Source/Editor/AnimGraph/Private/AnimGraphNode_StateMachineBase.cpp
--- a/Engine/Source/Editor/AnimGraph/Private/AnimGraphNode_StateMachineBase.cpp
+++ b/Engine/Source/Editor/AnimGraph/Private/AnimGraphNode_StateMachineBase.cpp
@@ -633,7 +633,9 @@ void UAnimGraphNode_StateMachineBase::OnProcessDuringCompilation(IAnimBlueprintC
 			BakedTransition.BlendMode = TransitionNode->BlendMode;
 			BakedTransition.CustomCurve = TransitionNode->CustomBlendCurve;
 			BakedTransition.BlendProfile = TransitionNode->BlendProfileWrapper.GetBlendProfile();
+			BakedTransition.MinTimeBeforeReentry = TransitionNode->MinTimeBeforeReentry;
 			BakedTransition.LogicType = TransitionNode->LogicType;
+			BakedTransition.bAllowInertializationForSelfTransitions = TransitionNode->bAllowInertializationForSelfTransitions;
 
 			UAnimStateNodeBase* PreviousState = StateNode;
 			UAnimStateNodeBase* NextState = TransitionNode->GetNextState();
diff --git a/Engine/Source/Editor/AnimGraph/Private/AnimStateEntryNode.cpp b/Engine/Source/Editor/AnimGraph/Private/AnimStateEntryNode.cpp
--- a/Engine/Source/Editor/AnimGraph/Private/AnimStateEntryNode.cpp
+++ b/Engine/Source/Editor/AnimGraph/Private/AnimStateEntryNode.cpp
@@ -53,4 +53,9 @@ UEdGraphNode* UAnimStateEntryNode::GetOutputNode() const
 	return NULL;
 }
 
+UEdGraphPin* UAnimStateEntryNode::GetOutputPin() const
+{
+	return FindPinChecked(TEXT("Entry"), EGPD_Output);
+}
+
 #undef LOCTEXT_NAMESPACE
diff --git a/Engine/Source/Editor/AnimGraph/Private/AnimStateNode.cpp b/Engine/Source/Editor/AnimGraph/Private/AnimStateNode.cpp
--- a/Engine/Source/Editor/AnimGraph/Private/AnimStateNode.cpp
+++ b/Engine/Source/Editor/AnimGraph/Private/AnimStateNode.cpp
@@ -11,6 +11,9 @@
 #include "Kismet2/Kismet2NameValidators.h"
 #include "AnimGraphNode_StateResult.h"
 #include "AnimationStateGraphSchema.h"
+#include "AnimGraphCommands.h"
+#include "ToolMenu.h"
+#include "ToolMenuSection.h"
 
 #define LOCTEXT_NAMESPACE "AnimStateNode"
 
@@ -94,6 +97,18 @@ UEdGraphPin* UAnimStateNode::GetPoseSinkPinInsideState() const
 	}
 }
 
+UAnimGraphNode_StateResult* UAnimStateNode::GetResultNodeInsideState() const
+{
+	if (UAnimationStateGraph* StateGraph = Cast<UAnimationStateGraph>(BoundGraph))
+	{
+		return StateGraph->MyResultNode;
+	}
+	else
+	{
+		return nullptr;
+	}
+}
+
 void UAnimStateNode::PostPasteNode()
 {
 	// Find an interesting name, but try to keep the same if possible
@@ -152,4 +167,14 @@ void UAnimStateNode::DestroyNode()
 	}
 }
 
+void UAnimStateNode::GetNodeContextMenuActions(UToolMenu* Menu, UGraphNodeContextMenuContext* Context) const
+{
+	// Add state node actions
+	{
+		FToolMenuSection& Section = Menu->AddSection("AnimStateNode", LOCTEXT("AnimStateNodeHeading", "State"));
+		Section.AddMenuEntry(FAnimGraphCommands::Get().SetActiveState);
+		Section.AddMenuEntry(FAnimGraphCommands::Get().CreateSelfTransition);
+	}
+}
+
 #undef LOCTEXT_NAMESPACE
diff --git a/Engine/Source/Editor/AnimGraph/Private/AnimStateTransitionNode.cpp b/Engine/Source/Editor/AnimGraph/Private/AnimStateTransitionNode.cpp
--- a/Engine/Source/Editor/AnimGraph/Private/AnimStateTransitionNode.cpp
+++ b/Engine/Source/Editor/AnimGraph/Private/AnimStateTransitionNode.cpp
@@ -12,6 +12,7 @@
 #include "AnimationTransitionSchema.h"
 #include "AnimationCustomTransitionGraph.h"
 #include "AnimationCustomTransitionSchema.h"
+#include "AnimGraphCommands.h"
 #include "AnimGraphNode_BlendSpacePlayer.h"
 #include "AnimGraphNode_SequencePlayer.h"
 #include "AnimGraphNode_StateResult.h"
@@ -22,6 +23,7 @@
 #include "EdGraphUtilities.h"
 #include "Kismet2/Kismet2NameValidators.h"
 #include "ScopedTransaction.h"
+#include "ToolMenu.h"
 #include "Animation/BlendProfile.h"
 #include "UObject/UE5MainStreamObjectVersion.h"
 
@@ -81,6 +83,7 @@ UAnimStateTransitionNode::UAnimStateTransitionNode(const FObjectInitializer& Obj
 	BlendMode = EAlphaBlendOption::HermiteCubic;
 	bAutomaticRuleBasedOnSequencePlayerInState = false;
 	AutomaticRuleTriggerTime = -1.f;
+	MinTimeBeforeReentry = -1.f;
 	bSharedRules = false;
 	SharedRulesGuid.Invalidate();
 	bSharedCrossfade = false;
@@ -90,6 +93,7 @@ UAnimStateTransitionNode::UAnimStateTransitionNode(const FObjectInitializer& Obj
 	bDisabled = false;
 	PriorityOrder = 1;
 	LogicType = ETransitionLogicType::TLT_StandardBlend;
+	bAllowInertializationForSelfTransitions = false;	// Defaults to false to preserve existing behavior
 }
 
 void UAnimStateTransitionNode::AllocateDefaultPins()
@@ -103,6 +107,9 @@ void UAnimStateTransitionNode::AllocateDefaultPins()
 void UAnimStateTransitionNode::PostPlacedNewNode()
 {
 	CreateBoundGraph();
+
+	// Only newly placed nodes fall back to inertialization for self transitions
+	bAllowInertializationForSelfTransitions = true;
 }
 
 void UAnimStateTransitionNode::PostLoad()
@@ -347,16 +354,15 @@ void UAnimStateTransitionNode::CreateConnections(UAnimStateNodeBase* PreviousSta
 
 void UAnimStateTransitionNode::RelinkHead(UAnimStateNodeBase* NewTargetState)
 {
-	UAnimStateNodeBase* SourceState = GetPreviousState();
 	UAnimStateNodeBase* TargetStateBeforeRelinking = GetNextState();
 
 	// Remove the incoming transition from the previous target state
 	TargetStateBeforeRelinking->GetInputPin()->Modify();
-	TargetStateBeforeRelinking->GetInputPin()->BreakLinkTo(SourceState->GetOutputPin());
+	TargetStateBeforeRelinking->GetInputPin()->BreakLinkTo(GetOutputPin());
 
 	// Add the new incoming transition to the new target state
 	NewTargetState->GetInputPin()->Modify();
-	NewTargetState->GetInputPin()->MakeLinkTo(SourceState->GetOutputPin());
+	NewTargetState->GetInputPin()->MakeLinkTo(GetOutputPin());
 
 	// Relink the target state of the transition node
 	Pins[1]->Modify();
@@ -364,8 +370,56 @@ void UAnimStateTransitionNode::RelinkHead(UAnimStateNodeBase* NewTargetState)
 	Pins[1]->MakeLinkTo(NewTargetState->GetInputPin());
 }
 
+void UAnimStateTransitionNode::RelinkTail(UAnimStateNodeBase* NewSourceState)
+{
+	UAnimStateNodeBase* SourceStateBeforeRelinking = GetPreviousState();
+
+	// Remove the outgoing transition from the previous source state
+	SourceStateBeforeRelinking->GetOutputPin()->Modify();
+	SourceStateBeforeRelinking->GetOutputPin()->BreakLinkTo(GetInputPin());
+
+	// Add the new outgoing transition to the new source state
+	NewSourceState->GetOutputPin()->Modify();
+	NewSourceState->GetOutputPin()->MakeLinkTo(GetInputPin());
+
+	// Relink the source state of the transition node
+	Pins[0]->Modify();
+	Pins[0]->BreakLinkTo(SourceStateBeforeRelinking->GetOutputPin());
+	Pins[0]->MakeLinkTo(NewSourceState->GetOutputPin());
+}
+
+void UAnimStateTransitionNode::Reverse()
+{
+	UAnimStateNodeBase* OldSourceState = GetPreviousState();
+	UAnimStateNodeBase* OldTargetState = GetNextState();
+	UAnimStateNodeBase* NewTargetState = OldSourceState;
+	UAnimStateNodeBase* NewSourceState = OldTargetState;
+
+	// Remove the incoming transition from the previous target state
+	OldSourceState->GetOutputPin()->Modify();
+	OldSourceState->GetOutputPin()->BreakLinkTo(OldTargetState->GetInputPin());
+
+	// Add the new incoming transition to the new target state
+	NewSourceState->GetOutputPin()->Modify();
+	NewSourceState->GetOutputPin()->MakeLinkTo(NewTargetState->GetInputPin());
+
+	// Relink the target state of the transition node
+	Pins[0]->Modify();
+	Pins[0]->BreakLinkTo(OldSourceState->GetOutputPin());
+	Pins[0]->MakeLinkTo(NewSourceState->GetOutputPin());
+
+	Pins[1]->Modify();
+	Pins[1]->BreakLinkTo(OldTargetState->GetInputPin());
+	Pins[1]->MakeLinkTo(NewTargetState->GetInputPin());
+}
+
 TArray<UAnimStateTransitionNode*> UAnimStateTransitionNode::GetListTransitionNodesToRelink(UEdGraphPin* SourcePin, UEdGraphPin* OldTargetPin, const TArray<UEdGraphNode*>& InSelectedGraphNodes)
 {
+	if (SourcePin->Direction == EGPD_Input && OldTargetPin->Direction == EGPD_Output)
+	{
+		Swap(SourcePin, OldTargetPin);
+	}
+	
 	UAnimStateNodeBase* SourceState = Cast<UAnimStateNodeBase>(SourcePin->GetOwningNode());
 	if (SourceState == nullptr || SourceState->GetInputPin() == nullptr || SourceState->GetOutputPin() == nullptr)
 	{
@@ -383,8 +437,7 @@ TArray<UAnimStateTransitionNode*> UAnimStateTransitionNode::GetListTransitionNod
 
 		// Get the actual target states from the transition nodes
 		UEdGraphNode* TransitionTargetNode = CurrentTransition->GetNextState();
-		UAnimStateTransitionNode* CastedOldTarget = Cast<UAnimStateTransitionNode>(OldTargetPin->GetOwningNode());
-		UEdGraphNode* OldTargetNode = CastedOldTarget->GetNextState();
+		UEdGraphNode* OldTargetNode = OldTargetPin->GetOwningNode();
 
 		// Compare the target states rather than comparing against the transition nodes
 		if (TransitionTargetNode != OldTargetNode)
@@ -832,6 +885,15 @@ TArray<UEdGraph*> UAnimStateTransitionNode::GetSubGraphs() const
 	return SubGraphs; 
 }
 
+void UAnimStateTransitionNode::GetNodeContextMenuActions(UToolMenu* Menu, UGraphNodeContextMenuContext* Context) const
+{
+	// Add transition actions
+	{
+		FToolMenuSection& Section = Menu->AddSection("AnimTransitionNode", LOCTEXT("AnimTransitionNodeHeading", "Transition"));
+		Section.AddMenuEntry(FAnimGraphCommands::Get().ReverseTransition);
+	}
+}
+
 //////////////////////////////////////////////////////////////////////////
 // IAnimStateTransitionNodeSharedDataHelper
 
diff --git a/Engine/Source/Editor/AnimGraph/Private/AnimationStateMachineSchema.cpp b/Engine/Source/Editor/AnimGraph/Private/AnimationStateMachineSchema.cpp
--- a/Engine/Source/Editor/AnimGraph/Private/AnimationStateMachineSchema.cpp
+++ b/Engine/Source/Editor/AnimGraph/Private/AnimationStateMachineSchema.cpp
@@ -274,6 +274,7 @@ bool UAnimationStateMachineSchema::TryRelinkConnectionTarget(UEdGraphPin* Source
 		return false;
 	}
 
+	bool bRelinkHead = SourcePin->Direction == EGPD_Output;
 	UAnimStateNodeBase* OldTargetState = Cast<UAnimStateNodeBase>(OldTargetPin->GetOwningNode());
 	UAnimStateNodeBase* NewTargetState = Cast<UAnimStateNodeBase>(NewTargetPin->GetOwningNode());
 	if (OldTargetState == nullptr || OldTargetState->GetInputPin() == nullptr || OldTargetState->GetOutputPin() == nullptr ||
@@ -303,9 +304,20 @@ bool UAnimationStateMachineSchema::TryRelinkConnectionTarget(UEdGraphPin* Source
 
 	// Collect all transition nodes starting at the source state, filter them by the transitions and perform the actual relink operation.
 	const TArray<UAnimStateTransitionNode*> TransitionNodes = UAnimStateTransitionNode::GetListTransitionNodesToRelink(SourcePin, OldTargetPin, InSelectedGraphNodes);
-	for (UAnimStateTransitionNode* TransitionNode : TransitionNodes)
+
+	if (bRelinkHead)
+	{
+		for (UAnimStateTransitionNode* TransitionNode : TransitionNodes)
+		{
+			TransitionNode->RelinkHead(NewTargetState);
+		}
+	}
+	else
 	{
-		TransitionNode->RelinkHead(NewTargetState);
+		for (UAnimStateTransitionNode* TransitionNode : TransitionNodes)
+		{
+			TransitionNode->RelinkTail(NewTargetState);
+		}
 	}
 
 	// In case one or more transitions got relinked, inform the blueprint about the changes
diff --git a/Engine/Source/Editor/AnimGraph/Private/K2Node_AnimGetter.cpp b/Engine/Source/Editor/AnimGraph/Private/K2Node_AnimGetter.cpp
--- a/Engine/Source/Editor/AnimGraph/Private/K2Node_AnimGetter.cpp
+++ b/Engine/Source/Editor/AnimGraph/Private/K2Node_AnimGetter.cpp
@@ -523,6 +523,10 @@ FText UK2Node_AnimGetter::GenerateTitle(UFunction* Getter, UAnimStateNodeBase* S
 	{
 		if (GetterRequiresParameter(Getter, TEXT("StateIndex")) || GetterRequiresParameter(Getter, TEXT("TransitionIndex")))
 		{
+			if (!SourceStateNode)
+			{
+				return FText::Format(InvalidNodeTitle, Getter->GetDisplayNameText());
+			}
 			return  FText::Format(LOCTEXT("NodeTitle", "{0} ({1})"), Getter->GetDisplayNameText(), SourceStateNode->GetNodeTitle(ENodeTitleType::ListView));
 		}
 		else
@@ -542,7 +546,10 @@ FText UK2Node_AnimGetter::GenerateTitle(UFunction* Getter, UAnimStateNodeBase* S
 
 void UK2Node_AnimGetter::UpdateCachedTitle()
 {
-	CachedTitle = GenerateTitle(GetTargetFunction(), SourceStateNode, SourceNode);
+	if (UFunction* Function = GetTargetFunction())
+	{
+		CachedTitle = GenerateTitle(Function, SourceStateNode, SourceNode);
+	}
 }
 
 FNodeSpawnData::FNodeSpawnData()
diff --git a/Engine/Source/Editor/AnimGraph/Public/AnimGraphCommands.h b/Engine/Source/Editor/AnimGraph/Public/AnimGraphCommands.h
--- a/Engine/Source/Editor/AnimGraph/Public/AnimGraphCommands.h
+++ b/Engine/Source/Editor/AnimGraph/Public/AnimGraphCommands.h
@@ -60,4 +60,13 @@ public:
 
 	// option for opening the asset related to the graph node
 	TSharedPtr< FUICommandInfo > OpenRelatedAsset;
+
+	// Create an alias for this state adds a transition to it
+	TSharedPtr< FUICommandInfo > CreateSelfTransition;
+
+	// Set the current state as active
+	TSharedPtr< FUICommandInfo > SetActiveState;
+	
+	// Reverse the selected transition(s)
+	TSharedPtr< FUICommandInfo > ReverseTransition;
 };
diff --git a/Engine/Source/Editor/AnimGraph/Public/AnimGraphNode_StateResult.h b/Engine/Source/Editor/AnimGraph/Public/AnimGraphNode_StateResult.h
--- a/Engine/Source/Editor/AnimGraph/Public/AnimGraphNode_StateResult.h
+++ b/Engine/Source/Editor/AnimGraph/Public/AnimGraphNode_StateResult.h
@@ -46,7 +46,7 @@ class UAnimGraphNode_StateResult : public UAnimGraphNode_Base
 	//~ End UAnimGraphNode_Base Interface
 
 	/** Function called when the owning state is entered, meaning it becomes the state machine's current state. */
-	UPROPERTY(EditAnywhere, Category = "Functions|State", meta=(FunctionReference, AllowFunctionLibraries, PrototypeFunction="/Script/AnimGraphRuntime.AnimExecutionContextLibrary.Prototype_ThreadSafeAnimUpdateCall", DefaultBindingName="OnStateEntry"), DisplayName="On State Entry")
+	UPROPERTY(EditAnywhere, Category = "Functions", meta=(FunctionReference, AllowFunctionLibraries, PrototypeFunction="/Script/AnimGraphRuntime.AnimExecutionContextLibrary.Prototype_ThreadSafeAnimUpdateCall", DefaultBindingName="OnStateEntry"), DisplayName="On State Entry")
 	FMemberReference StateEntryFunction;
 
 	/** Function called when the owning state is fully blended in.
@@ -55,15 +55,11 @@ class UAnimGraphNode_StateResult : public UAnimGraphNode_Base
 	 * - This is only called for the state machine's current state since its the most recent transition's target state.
 	 * - This will not be called if the state is skipped. This can happen when the flag bSkipFirstUpdateTransition on the state machine node is set to true.
 	 */
-	UPROPERTY(EditAnywhere, Category = "Functions|State", meta=(FunctionReference, AllowFunctionLibraries, PrototypeFunction="/Script/AnimGraphRuntime.AnimExecutionContextLibrary.Prototype_ThreadSafeAnimUpdateCall", DefaultBindingName="OnStateFullyBlendedIn"), DisplayName="On State Fully Blended In")
+	UPROPERTY(EditAnywhere, Category = "Functions", meta=(FunctionReference, AllowFunctionLibraries, PrototypeFunction="/Script/AnimGraphRuntime.AnimExecutionContextLibrary.Prototype_ThreadSafeAnimUpdateCall", DefaultBindingName="OnStateFullyBlendedIn"), DisplayName="On State Fully Blended In")
 	FMemberReference StateFullyBlendedInFunction;
 
-	/** Function called when the owning state is exited, meaning it stops being the state machine's current state.
-	 *
-	 * Notes:
-	 * - This will not be called if the state machine node loses relevancy. Please use "On State Interrupt" for that case.
-	 */
-	UPROPERTY(EditAnywhere, Category = "Functions|State", meta=(FunctionReference, AllowFunctionLibraries, PrototypeFunction="/Script/AnimGraphRuntime.AnimExecutionContextLibrary.Prototype_ThreadSafeAnimUpdateCall", DefaultBindingName="OnStateExit"), DisplayName="On State Exit")
+	/** Function called when the owning state is exited, meaning it stops being the state machine's current state. */
+	UPROPERTY(EditAnywhere, Category = "Functions", meta=(FunctionReference, AllowFunctionLibraries, PrototypeFunction="/Script/AnimGraphRuntime.AnimExecutionContextLibrary.Prototype_ThreadSafeAnimUpdateCall", DefaultBindingName="OnStateExit"), DisplayName="On State Exit")
 	FMemberReference StateExitFunction;
 
 	/** Function called when the owning state is fully blended out.
@@ -71,8 +67,7 @@ class UAnimGraphNode_StateResult : public UAnimGraphNode_Base
 	 * Notes:
 	 * - This will be called for any states that had weight.
 	 * - This will not be called if the state is skipped. This can happen when the flag bSkipFirstUpdateTransition on the state machine node is set to true.
-	 * - This will not be called if the state machine node loses relevancy. Please use "On State Interrupt" for that case.
 	 */
-	UPROPERTY(EditAnywhere, Category = "Functions|State", meta=(FunctionReference, AllowFunctionLibraries, PrototypeFunction="/Script/AnimGraphRuntime.AnimExecutionContextLibrary.Prototype_ThreadSafeAnimUpdateCall", DefaultBindingName="OnStateFullyBlendedOut"), DisplayName="On State Fully Blended Out")
+	UPROPERTY(EditAnywhere, Category = "Functions", meta=(FunctionReference, AllowFunctionLibraries, PrototypeFunction="/Script/AnimGraphRuntime.AnimExecutionContextLibrary.Prototype_ThreadSafeAnimUpdateCall", DefaultBindingName="OnStateFullyBlendedOut"), DisplayName="On State Fully Blended Out")
 	FMemberReference StateFullyBlendedOutFunction;
 };
diff --git a/Engine/Source/Editor/AnimGraph/Public/AnimStateEntryNode.h b/Engine/Source/Editor/AnimGraph/Public/AnimStateEntryNode.h
--- a/Engine/Source/Editor/AnimGraph/Public/AnimStateEntryNode.h
+++ b/Engine/Source/Editor/AnimGraph/Public/AnimStateEntryNode.h
@@ -22,8 +22,10 @@ class UAnimStateEntryNode : public UEdGraphNode
 	virtual void AllocateDefaultPins() override;
 	virtual FText GetNodeTitle(ENodeTitleType::Type TitleType) const override;
 	virtual FText GetTooltipText() const override;
+	virtual bool CanUserDeleteNode() const override { return false; }
+	virtual bool CanDuplicateNode() const override { return false; }
 	//~ End UEdGraphNode Interface
-	
-	ANIMGRAPH_API UEdGraphNode* GetOutputNode() const;
 
+	ANIMGRAPH_API UEdGraphNode* GetOutputNode() const;
+	ANIMGRAPH_API UEdGraphPin* GetOutputPin() const;
 };
diff --git a/Engine/Source/Editor/AnimGraph/Public/AnimStateNode.h b/Engine/Source/Editor/AnimGraph/Public/AnimStateNode.h
--- a/Engine/Source/Editor/AnimGraph/Public/AnimStateNode.h
+++ b/Engine/Source/Editor/AnimGraph/Public/AnimStateNode.h
@@ -10,6 +10,7 @@
 
 class UEdGraph;
 class UEdGraphPin;
+class UAnimGraphNode_StateResult;
 
 UENUM()
 enum EAnimStateType : int
@@ -33,17 +34,17 @@ public:
 	UPROPERTY()
 	TEnumAsByte<EAnimStateType> StateType;
 
-	UPROPERTY(EditAnywhere, Category=Events)
+	UPROPERTY(EditAnywhere, Category = "Deprecated")
 	FAnimNotifyEvent StateEntered;
 	
-	UPROPERTY(EditAnywhere, Category=Events)
+	UPROPERTY(EditAnywhere, Category = "Deprecated")
 	FAnimNotifyEvent StateLeft;
 
-	UPROPERTY(EditAnywhere, Category=Events)
+	UPROPERTY(EditAnywhere, Category = "Deprecated")
 	FAnimNotifyEvent StateFullyBlended;
 
 	// Whether or not this state will ALWAYS reset it's state on reentry, regardless of remaining weight
-	UPROPERTY(EditAnywhere, AdvancedDisplay, Category = Events)
+	UPROPERTY(EditAnywhere, Category = "Animation State")
 	bool bAlwaysResetOnEntry;
 
 	//~ Begin UObject Interface
@@ -60,6 +61,7 @@ public:
 	virtual void PostPlacedNewNode() override;
 	virtual void DestroyNode() override;
 	virtual TArray<UEdGraph*> GetSubGraphs() const override { return TArray<UEdGraph*>( { BoundGraph } ); }
+	virtual void GetNodeContextMenuActions(UToolMenu* Menu, UGraphNodeContextMenuContext* Context) const override;
 	//~ End UEdGraphNode Interface
 	
 	//~ Begin UAnimStateNodeBase Interface
@@ -71,6 +73,9 @@ public:
 	// @return the pose pin of the state sink node within the anim graph of this state
 	ANIMGRAPH_API UEdGraphPin* GetPoseSinkPinInsideState() const;
 
+	// @return the state sink node within the anim graph of this state
+	ANIMGRAPH_API UAnimGraphNode_StateResult* GetResultNodeInsideState() const;
+
 public:
 	virtual UEdGraph* GetBoundGraph() const override { return BoundGraph; }
 	virtual void ClearBoundGraph() override { BoundGraph = nullptr; }
diff --git a/Engine/Source/Editor/AnimGraph/Public/AnimStateTransitionNode.h b/Engine/Source/Editor/AnimGraph/Public/AnimStateTransitionNode.h
--- a/Engine/Source/Editor/AnimGraph/Public/AnimStateTransitionNode.h
+++ b/Engine/Source/Editor/AnimGraph/Public/AnimStateTransitionNode.h
@@ -35,7 +35,7 @@ class UAnimStateTransitionNode : public UAnimStateNodeBase
 	int32 PriorityOrder;
 
 	// The duration to cross-fade for
-	UPROPERTY(EditAnywhere, Config, Category=Transition, Meta=(ClampMin="0.0"))
+	UPROPERTY(EditAnywhere, Config, Category=Transition, Meta=(ClampMin="0.0", ForceUnits="s"))
 	float CrossfadeDuration;
 
 	// The type of blending to use in the crossfade
@@ -64,9 +64,15 @@ class UAnimStateTransitionNode : public UAnimStateNodeBase
 	// When should the automatic transition rule trigger relative to the time remaining on the relevant asset player:
 	//  < 0 means trigger the transition 'Crossfade Duration' seconds before the end of the asset player, so a standard blend would finish just as the asset player ends
 	// >= 0 means trigger the transition 'Automatic Rule Trigger Time' seconds before the end of the asset player
-	UPROPERTY(EditAnywhere, Category=Transition)
+	UPROPERTY(EditAnywhere, Category=Transition, meta=(ClampMin=0.0, ForceUnits="s"))
 	float AutomaticRuleTriggerTime;
 
+	// The minimum time elapsed (in seconds) before the target state can be re-entered via this transition 
+	// Has no effect when set to -1.0.
+	// When set to zero, wait at least one frame before re-entry via this transition is allowed.
+	UPROPERTY(EditAnywhere, Category=Transition, DisplayName="Min Time Before Re-entry", meta=(ClampMin=0.0, ForceUnits="s"))
+	float MinTimeBeforeReentry;
+
 	// If SyncGroupName is specified, Transition will only be taken if the SyncGroup has valid markers (other transition rules still apply in addition to that).
 	UPROPERTY(EditAnywhere, Category = Transition)
 	FName SyncGroupNameToRequireValidMarkersRule;
@@ -84,6 +90,10 @@ class UAnimStateTransitionNode : public UAnimStateNodeBase
 	UPROPERTY(EditAnywhere, Category=Events)
 	FAnimNotifyEvent TransitionInterrupt;
 
+	/** Whether to fall back to inertialization/dead blending when reentering an already-active state. This can avoid pops. */
+	UPROPERTY(EditAnywhere, Category=Transition)
+	bool bAllowInertializationForSelfTransitions;
+
 	/** This transition can go both directions */
 	UPROPERTY(EditAnywhere, Category=Transition)
 	bool Bidirectional;
@@ -121,6 +131,9 @@ class UAnimStateTransitionNode : public UAnimStateNodeBase
 	UPROPERTY()
 	int32 SharedCrossfadeIdx;
 
+	UPROPERTY()
+	FVector2D CachedRotation;
+
 	//~ Begin UObject Interface
 	virtual void PostEditChangeProperty(struct FPropertyChangedEvent& PropertyChangedEvent) override;
 	virtual void PostLoad() override;
@@ -141,6 +154,7 @@ class UAnimStateTransitionNode : public UAnimStateNodeBase
 	virtual void ValidateNodeDuringCompilation(class FCompilerResultsLog& MessageLog) const override;
 	virtual UObject* GetJumpTargetForDoubleClick() const override;
 	virtual TArray<UEdGraph*> GetSubGraphs() const override;
+	virtual void GetNodeContextMenuActions(UToolMenu* Menu, UGraphNodeContextMenuContext* Context) const override;
 	//~ End UEdGraphNode Interface
 
 	//~ Begin UAnimStateNodeBase Interface
@@ -158,11 +172,22 @@ class UAnimStateTransitionNode : public UAnimStateNodeBase
 	ANIMGRAPH_API void CreateConnections(UAnimStateNodeBase* PreviousState, UAnimStateNodeBase* NextState);
 
 	/**
-	 * Relink transition head (where the arrow is of a state transition) to a new state.
+	 * Relink transition head (where the arrow head is of a state transition) to a new state.
 	 * @param[in] NewTargetState The new transition target.
 	 */
 	ANIMGRAPH_API void RelinkHead(UAnimStateNodeBase* NewTargetState);
 
+	/**
+	 * Relink transition tail (where the arrow base is of a state transition) to a new state.
+	 * @param[in] NewSourceState The new transition source.
+	 */
+	ANIMGRAPH_API void RelinkTail(UAnimStateNodeBase* NewSourceState);
+
+	/**
+	 * Reverse the direction of this transition
+	 */
+	ANIMGRAPH_API void Reverse();
+	
 	/**
 	 * Helper function to gather the transition nodes to be relinked by taking the graph selection into account as well.
 	 * For example when relinking a transition holding several transition nodes but only a few are selected to be relinked.
diff --git a/Engine/Source/Editor/AnimationBlueprintEditor/Private/AnimationBlueprintEditor.cpp b/Engine/Source/Editor/AnimationBlueprintEditor/Private/AnimationBlueprintEditor.cpp
--- a/Engine/Source/Editor/AnimationBlueprintEditor/Private/AnimationBlueprintEditor.cpp
+++ b/Engine/Source/Editor/AnimationBlueprintEditor/Private/AnimationBlueprintEditor.cpp
@@ -99,6 +99,9 @@
 #include "PersonaUtils.h"
 // Hide related nodes feature
 #include "AnimationBlueprintToolMenuContext.h"
+#include "AnimStateAliasNode.h"
+#include "AnimStateNode.h"
+#include "AnimStateTransitionNode.h"
 #include "DetailLayoutBuilder.h"
 #include "Preferences/AnimationBlueprintEditorOptions.h"
 #include "PropertyEditorDelegates.h"
@@ -118,6 +121,7 @@
 #include "Templates/SubclassOf.h"
 #include "ToolMenuContext.h"
 #include "ToolMenus.h"
+#include "Animation/AnimNode_StateMachine.h"
 #include "Toolkits/AssetEditorToolkit.h"
 #include "UObject/Object.h"
 #include "UObject/ObjectPtr.h"
@@ -761,6 +765,21 @@ void FAnimationBlueprintEditor::OnCreateGraphEditorCommands(TSharedPtr<FUIComman
 	GraphEditorCommandsList->MapAction( FAnimGraphCommands::Get().OpenRelatedAsset,
 		FExecuteAction::CreateSP( this, &FAnimationBlueprintEditor::OnOpenRelatedAsset )
 		);
+
+	GraphEditorCommandsList->MapAction( FAnimGraphCommands::Get().CreateSelfTransition,
+		FExecuteAction::CreateSP( this, &FAnimationBlueprintEditor::OnCreateSelfTransition )
+		);
+
+	GraphEditorCommandsList->MapAction( FAnimGraphCommands::Get().SetActiveState,
+		FUIAction(
+			FExecuteAction::CreateSP( this, &FAnimationBlueprintEditor::OnSetActiveState ),
+			FCanExecuteAction::CreateSP( this, &FAnimationBlueprintEditor::CanSetActiveState )
+			)
+		);
+
+	GraphEditorCommandsList->MapAction( FAnimGraphCommands::Get().ReverseTransition,
+		FExecuteAction::CreateSP( this, &FAnimationBlueprintEditor::OnReverseTransition )
+		);
 }
 
 
@@ -1527,6 +1546,135 @@ void FAnimationBlueprintEditor::OnOpenRelatedAsset()
 	}
 }
 
+void FAnimationBlueprintEditor::OnCreateSelfTransition()
+{
+	FGraphPanelSelectionSet SelectedNodes = GetSelectedNodes();
+	if (SelectedNodes.Num() == 0)
+	{
+		return;
+	}
+
+	const FScopedTransaction Transaction(LOCTEXT("CreateSelfTransition", "Create Self Transition"));
+
+	for (UObject* Node : SelectedNodes)
+	{
+		if (UAnimStateNode* StateNode = Cast<UAnimStateNode>(Node))
+		{
+			UEdGraph* TargetGraph = StateNode->GetGraph();
+			TargetGraph->Modify();
+
+			FGraphNodeCreator<UAnimStateAliasNode> AliasNodeCreator(*TargetGraph);
+			UAnimStateAliasNode* AliasNode = AliasNodeCreator.CreateNode();
+			TSharedPtr<INameValidatorInterface> NameValidator = AliasNode->MakeNameValidator();
+			check(NameValidator.IsValid());
+			FString NewName(TEXT("Self"));
+			NameValidator->FindValidString(NewName);
+			AliasNode->StateAliasName = NewName;
+			AliasNodeCreator.Finalize();
+
+			AliasNode->SetNodePosX(StateNode->GetNodePosX() + 200);
+			AliasNode->SetNodePosY(StateNode->GetNodePosY() - 100);
+			AliasNode->GetAliasedStates().Add(StateNode);
+			
+			FGraphNodeCreator<UAnimStateTransitionNode> TransitionNodeCreator(*TargetGraph);
+			UAnimStateTransitionNode* TransitionNode = TransitionNodeCreator.CreateNode();
+			TransitionNodeCreator.Finalize();
+
+			// Self-transitions should always fall back to inertialization
+			TransitionNode->bAllowInertializationForSelfTransitions = true;
+			// Self transitions should always have a delay so they dont re-enter each frame (footgun prevention)
+			TransitionNode->MinTimeBeforeReentry = 0.0f;
+			TransitionNode->CreateConnections(AliasNode, StateNode);
+		}
+	}
+}
+
+FAnimNode_StateMachine* FAnimationBlueprintEditor::GetSelectedStateInstance(int32& OutSelectedState) const
+{
+	FGraphPanelSelectionSet SelectedNodes = GetSelectedNodes();
+	if (SelectedNodes.Num() != 1)
+	{
+		return nullptr;
+	}
+
+	UAnimStateNode* StateNode = Cast<UAnimStateNode>(*SelectedNodes.begin());
+	if (StateNode == nullptr)
+	{
+		return nullptr;
+	}
+
+	UAnimBlueprint* AnimBlueprint = Cast<UAnimBlueprint>(FBlueprintEditorUtils::FindBlueprintForNode(StateNode));
+	if(AnimBlueprint == nullptr)
+	{
+		return nullptr;
+	}
+
+	UAnimInstance* ActiveObject = Cast<UAnimInstance>(AnimBlueprint->GetObjectBeingDebugged());
+	UAnimBlueprintGeneratedClass* Class = AnimBlueprint->GetAnimBlueprintGeneratedClass();
+
+	if ((ActiveObject == nullptr) || (Class == nullptr) || Class->GetAnimNodeProperties().Num() == 0)
+	{
+		return nullptr;
+	}
+
+	FStateMachineDebugData* DebugInfo = Class->GetAnimBlueprintDebugData().StateMachineDebugData.Find(StateNode->GetGraph());
+	if (DebugInfo == nullptr)
+	{
+		return nullptr;
+	}
+
+	int32* StateIndexPtr = DebugInfo->NodeToStateIndex.Find(StateNode);
+	if(StateIndexPtr == nullptr)
+	{
+		return nullptr;
+	}
+
+	if (!Class->BakedStateMachines.IsValidIndex(DebugInfo->MachineIndex))
+	{
+		return nullptr;
+	}
+
+	OutSelectedState = *StateIndexPtr;
+	return const_cast<FAnimNode_StateMachine*>(ActiveObject->GetStateMachineInstanceFromName(Class->BakedStateMachines[DebugInfo->MachineIndex].MachineName));
+}
+
+bool FAnimationBlueprintEditor::CanSetActiveState() const
+{
+	int32 StateIndex = 0;
+	return GetSelectedStateInstance(StateIndex) != nullptr;
+}
+
+void FAnimationBlueprintEditor::OnSetActiveState()
+{
+	int32 StateIndex = 0;
+	FAnimNode_StateMachine* StateMachineInstance = GetSelectedStateInstance(StateIndex);
+	if (StateMachineInstance == nullptr)
+	{
+		return;
+	}
+
+	StateMachineInstance->DebugForceState(StateIndex);
+}
+
+void FAnimationBlueprintEditor::OnReverseTransition()
+{
+	FGraphPanelSelectionSet SelectedNodes = GetSelectedNodes();
+	if (SelectedNodes.Num() == 0)
+	{
+		return;
+	}
+
+	const FScopedTransaction Transaction(LOCTEXT("ReverseTransition", "Reverse Transition"));
+
+	for (UObject* Node : SelectedNodes)
+	{
+		if (UAnimStateTransitionNode* TransitionNode = Cast<UAnimStateTransitionNode>(Node))
+		{
+			TransitionNode->Reverse();
+		}
+	}
+}
+
 bool FAnimationBlueprintEditor::CanRemovePosePin() const
 {
 	return true;
diff --git a/Engine/Source/Editor/AnimationBlueprintEditor/Private/AnimationBlueprintEditor.h b/Engine/Source/Editor/AnimationBlueprintEditor/Private/AnimationBlueprintEditor.h
--- a/Engine/Source/Editor/AnimationBlueprintEditor/Private/AnimationBlueprintEditor.h
+++ b/Engine/Source/Editor/AnimationBlueprintEditor/Private/AnimationBlueprintEditor.h
@@ -45,6 +45,7 @@ class USkeletalMeshComponent;
 struct FToolMenuContext;
 struct FFrame;
 struct FBlueprintExceptionInfo;
+struct FAnimNode_StateMachine;
 
 struct FAnimationBlueprintEditorModes
 {
@@ -383,6 +384,13 @@ private:
 	// Opens the associated asset of the selected nodes
 	void OnOpenRelatedAsset();
 
+	// State machine functions
+	void OnCreateSelfTransition();
+	FAnimNode_StateMachine* GetSelectedStateInstance(int32& OutSelectedState) const;
+	bool CanSetActiveState() const;
+	void OnSetActiveState();
+	void OnReverseTransition();
+
 	/** The extender to pass to the level editor to extend it's window menu */
 	TSharedPtr<FExtender> MenuExtender;
 
diff --git a/Engine/Source/Editor/AnimationBlueprintEditor/Private/AnimationStateNodes/SGraphNodeAnimState.cpp b/Engine/Source/Editor/AnimationBlueprintEditor/Private/AnimationStateNodes/SGraphNodeAnimState.cpp
--- a/Engine/Source/Editor/AnimationBlueprintEditor/Private/AnimationStateNodes/SGraphNodeAnimState.cpp
+++ b/Engine/Source/Editor/AnimationBlueprintEditor/Private/AnimationStateNodes/SGraphNodeAnimState.cpp
@@ -3,6 +3,7 @@
 #include "AnimationStateNodes/SGraphNodeAnimState.h"
 
 #include "AnimStateConduitNode.h"
+#include "AnimStateAliasNode.h"
 #include "AnimStateNodeBase.h"
 #include "Animation/AnimBlueprint.h"
 #include "Animation/AnimBlueprintGeneratedClass.h"
@@ -29,6 +30,7 @@
 #include "SGraphPreviewer.h"
 #include "SNodePanel.h"
 #include "SlotBase.h"
+#include "SStateMachineInputPin.h"
 #include "Styling/AppStyle.h"
 #include "Templates/Casts.h"
 #include "Types/SlateEnums.h"
@@ -40,6 +42,7 @@
 #include "Widgets/SToolTip.h"
 #include "Widgets/Text/SInlineEditableTextBlock.h"
 #include "Widgets/Text/STextBlock.h"
+#include "SStateMachineOutputPin.h"
 
 class SWidget;
 class UEdGraphSchema;
@@ -49,57 +52,6 @@ struct FSlateBrush;
 
 #define LOCTEXT_NAMESPACE "SGraphNodeAnimState"
 
-/////////////////////////////////////////////////////
-// SStateMachineOutputPin
-
-class SStateMachineOutputPin : public SGraphPin
-{
-public:
-	SLATE_BEGIN_ARGS(SStateMachineOutputPin){}
-	SLATE_END_ARGS()
-
-	void Construct(const FArguments& InArgs, UEdGraphPin* InPin);
-protected:
-	// Begin SGraphPin interface
-	virtual TSharedRef<SWidget>	GetDefaultValueWidget() override;
-	// End SGraphPin interface
-
-	const FSlateBrush* GetPinBorder() const;
-};
-
-void SStateMachineOutputPin::Construct(const FArguments& InArgs, UEdGraphPin* InPin)
-{
-	this->SetCursor( EMouseCursor::Default );
-
-	bShowLabel = true;
-
-	GraphPinObj = InPin;
-	check(GraphPinObj != NULL);
-
-	const UEdGraphSchema* Schema = GraphPinObj->GetSchema();
-	check(Schema);
-
-	// Set up a hover for pins that is tinted the color of the pin.
-	SBorder::Construct( SBorder::FArguments()
-		.BorderImage( this, &SStateMachineOutputPin::GetPinBorder )
-		.BorderBackgroundColor( this, &SStateMachineOutputPin::GetPinColor )
-		.OnMouseButtonDown( this, &SStateMachineOutputPin::OnPinMouseDown )
-		.Cursor( this, &SStateMachineOutputPin::GetPinCursor )
-	);
-}
-
-TSharedRef<SWidget>	SStateMachineOutputPin::GetDefaultValueWidget()
-{
-	return SNew(STextBlock);
-}
-
-const FSlateBrush* SStateMachineOutputPin::GetPinBorder() const
-{
-	return ( IsHovered() )
-		? FAppStyle::GetBrush( TEXT("Graph.StateNode.Pin.BackgroundHovered") )
-		: FAppStyle::GetBrush( TEXT("Graph.StateNode.Pin.Background") );
-}
-
 /////////////////////////////////////////////////////
 // SGraphNodeAnimState
 
@@ -112,7 +64,7 @@ void SGraphNodeAnimState::Construct(const FArguments& InArgs, UAnimStateNodeBase
 	this->UpdateGraphNode();
 }
 
-void SGraphNodeAnimState::GetStateInfoPopup(UEdGraphNode* GraphNode, TArray<FGraphInformationPopupInfo>& Popups)
+void SGraphNodeAnimState::GetStateInfoPopup(TArray<FGraphInformationPopupInfo>& Popups) const
 {
 	UAnimBlueprint* AnimBlueprint = Cast<UAnimBlueprint>(FBlueprintEditorUtils::FindBlueprintForNode(GraphNode));
 	if(AnimBlueprint)
@@ -120,7 +72,7 @@ void SGraphNodeAnimState::GetStateInfoPopup(UEdGraphNode* GraphNode, TArray<FGra
 		UAnimInstance* ActiveObject = Cast<UAnimInstance>(AnimBlueprint->GetObjectBeingDebugged());
 		UAnimBlueprintGeneratedClass* Class = AnimBlueprint->GetAnimBlueprintGeneratedClass();
 
-		FLinearColor CurrentStateColor(1.f, 0.5f, 0.25f);
+		FLinearColor CurrentStateColor = GetBorderBackgroundColor().GetSpecifiedColor();
 
 		// Display various types of debug data
 		if ((ActiveObject != NULL) && (Class != NULL))
@@ -161,14 +113,15 @@ void SGraphNodeAnimState::GetStateInfoPopup(UEdGraphNode* GraphNode, TArray<FGra
 
 void SGraphNodeAnimState::GetNodeInfoPopups(FNodeInfoContext* Context, TArray<FGraphInformationPopupInfo>& Popups) const
 {
-	GetStateInfoPopup(GraphNode, Popups);
+	GetStateInfoPopup(Popups);
 }
 
 FSlateColor SGraphNodeAnimState::GetBorderBackgroundColor() const
 {
 	FLinearColor InactiveStateColor(0.08f, 0.08f, 0.08f);
-	FLinearColor ActiveStateColorDim(0.4f, 0.3f, 0.15f);
-	FLinearColor ActiveStateColorBright(1.f, 0.6f, 0.35f);
+	FLinearColor ActiveStateColorDim = InactiveStateColor;
+	FLinearColor ActiveStateColorBright = FStyleColors::AccentOrange.GetSpecifiedColor().Desaturate(0.25f);
+	ActiveStateColorBright.A = 1.0f;
 
 	return GetBorderBackgroundColor_Internal(InactiveStateColor, ActiveStateColorDim, ActiveStateColorBright);
 }
@@ -256,84 +209,142 @@ void SGraphNodeAnimState::UpdateGraphNode()
 	LeftNodeBox.Reset();
 
 	const FSlateBrush* NodeTypeIcon = GetNameIcon();
+	const FSlateBrush* BodyBrush = FAppStyle::GetBrush( "Graph.AnimStateNode.Body" );
+	const FSlateBrush* SpillBrush = FAppStyle::GetBrush( "Graph.AnimStateNode.ColorSpill" );
+	const FSlateBrush* SelectionBrush = FAppStyle::GetBrush( "Graph.AnimStateNode.Selection" );
+
+	if (GraphNode->IsA<UAnimStateAliasNode>())
+	{
+		BodyBrush = FAppStyle::GetBrush( "Graph.AnimAliasNode.Body" );
+		SpillBrush = FAppStyle::GetBrush( "Graph.AnimAliasNode.ColorSpill" );
+		SelectionBrush = FAppStyle::GetBrush( "Graph.AnimAliasNode.Selection" );
+	}
+	else if (GraphNode->IsA<UAnimStateConduitNode>())
+	{
+		BodyBrush = FAppStyle::GetBrush( "Graph.AnimConduitNode.Body" );
+		SpillBrush = FAppStyle::GetBrush( "Graph.AnimConduitNode.ColorSpill" );
+		SelectionBrush = FAppStyle::GetBrush( "Graph.AnimConduitNode.Selection" );
+	}
 
 	FLinearColor TitleShadowColor(0.6f, 0.6f, 0.6f);
 	TSharedPtr<SErrorText> ErrorText;
 	TSharedPtr<SNodeTitle> NodeTitle = SNew(SNodeTitle, GraphNode);
 
+	// Set ourselves as SelfHitTestInvisible so we dont end up with a 1-unit border where the node is selected rather than the pin
+	SetVisibility(EVisibility::SelfHitTestInvisible);
+	
 	this->ContentScale.Bind( this, &SGraphNode::GetContentScale );
 	this->GetOrAddSlot( ENodeZone::Center )
 		.HAlign(HAlign_Center)
 		.VAlign(VAlign_Center)
 		[
-			SNew(SBorder)
-			.BorderImage( FAppStyle::GetBrush( "Graph.StateNode.Body" ) )
-			.Padding(0)
-			.BorderBackgroundColor( this, &SGraphNodeAnimState::GetBorderBackgroundColor )
+			SNew(SOverlay)
+			.Visibility(EVisibility::SelfHitTestInvisible)
+			+SOverlay::Slot()
+			.Padding(2.0f)
 			[
-				SNew(SOverlay)
-
-				// PIN AREA
-				+SOverlay::Slot()
-				.HAlign(HAlign_Fill)
-				.VAlign(VAlign_Fill)
+				SNew(SBorder)
+				.BorderImage( BodyBrush )
+				.Padding(0)
+				.BorderBackgroundColor( this, &SGraphNodeAnimState::GetBorderBackgroundColor )
 				[
-					SAssignNew(RightNodeBox, SVerticalBox)
-				]
+					SNew(SOverlay)
 
-				// STATE NAME AREA
-				+SOverlay::Slot()
-				.HAlign(HAlign_Center)
-				.VAlign(VAlign_Center)
-				.Padding(10.0f)
-				[
-					SNew(SBorder)
-					.BorderImage( FAppStyle::GetBrush("Graph.StateNode.ColorSpill") )
-					.BorderBackgroundColor( TitleShadowColor )
-					.HAlign(HAlign_Center)
-					.VAlign(VAlign_Center)
-					.Visibility(EVisibility::SelfHitTestInvisible)
+					// PIN AREA
+					+SOverlay::Slot()
+					.HAlign(HAlign_Fill)
+					.VAlign(VAlign_Fill)
 					[
-						SNew(SHorizontalBox)
-						+SHorizontalBox::Slot()
-						.AutoWidth()
+						SAssignNew(RightNodeBox, SVerticalBox)
+						+SVerticalBox::Slot()
+						.HAlign(HAlign_Fill)
+						.VAlign(VAlign_Fill)
+						.FillHeight(1.0f)
 						[
-							// POPUP ERROR MESSAGE
-							SAssignNew(ErrorText, SErrorText )
-							.BackgroundColor( this, &SGraphNodeAnimState::GetErrorColor )
-							.ToolTipText( this, &SGraphNodeAnimState::GetErrorMsgToolTip )
+							SAssignNew(PinOverlay, SOverlay)
 						]
-						+SHorizontalBox::Slot()
-						.AutoWidth()
+					]
+
+					// STATE NAME AREA
+					+SOverlay::Slot()
+					.HAlign(HAlign_Center)
+					.VAlign(VAlign_Center)
+					.Padding(12.0f)
+					[
+						SNew(SBorder)
+						.BorderImage( SpillBrush )
+						.BorderBackgroundColor( TitleShadowColor )
+						.HAlign(HAlign_Center)
 						.VAlign(VAlign_Center)
+						.Padding(FMargin(2.0f, 2.0f, 0.0f, 2.0f))
 						[
-							SNew(SImage)
-							.Image(NodeTypeIcon)
-						]
-						+SHorizontalBox::Slot()
-						.Padding(FMargin(4.0f, 0.0f, 4.0f, 0.0f))
-						[
-							SNew(SVerticalBox)
-							+SVerticalBox::Slot()
-								.AutoHeight()
+							SNew(SHorizontalBox)
+							+SHorizontalBox::Slot()
+							.AutoWidth()
 							[
-								SAssignNew(InlineEditableText, SInlineEditableTextBlock)
-								.Style( FAppStyle::Get(), "Graph.StateNode.NodeTitleInlineEditableText" )
-								.Text( NodeTitle.Get(), &SNodeTitle::GetHeadTitle )
-								.OnVerifyTextChanged(this, &SGraphNodeAnimState::OnVerifyNameTextChanged)
-								.OnTextCommitted(this, &SGraphNodeAnimState::OnNameTextCommited)
-								.IsReadOnly( this, &SGraphNodeAnimState::IsNameReadOnly )
-								.IsSelected(this, &SGraphNodeAnimState::IsSelectedExclusively)
+								// POPUP ERROR MESSAGE
+								SAssignNew(ErrorText, SErrorText )
+								.BackgroundColor( this, &SGraphNodeAnimState::GetErrorColor )
+								.ToolTipText( this, &SGraphNodeAnimState::GetErrorMsgToolTip )
 							]
-							+SVerticalBox::Slot()
-								.AutoHeight()
+							+SHorizontalBox::Slot()
+							.AutoWidth()
+							.VAlign(VAlign_Center)
+							.Padding(FMargin(7.0f, 4.0f, 2.0f, 4.0f))
 							[
-								NodeTitle.ToSharedRef()
+								SNew(SImage)
+								.Image(NodeTypeIcon)
+								.ColorAndOpacity( this, &SGraphNodeAnimState::GetIconColor )
+							]
+							+SHorizontalBox::Slot()
+							.VAlign(VAlign_Center)
+							.HAlign(HAlign_Center)
+							.AutoWidth()
+							[
+								SNew(SVerticalBox)
+								+SVerticalBox::Slot()
+								.AutoHeight()
+								.Padding(MakeAttributeLambda([this]() -> FMargin
+								{
+									return InlineEditableText->IsInEditMode() ?
+										FMargin(0.0f, 0.0f, 0.0f, 0.0f) :
+										FMargin(5.0f, 0.0f, 13.0f, 0.0f);
+								}))
+								[
+									SAssignNew(InlineEditableText, SInlineEditableTextBlock)
+									.Style( FAppStyle::Get(), "Graph.AnimStateNode.NodeTitleInlineEditableText" )
+									.Text( NodeTitle.Get(), &SNodeTitle::GetHeadTitle )
+									.OnVerifyTextChanged(this, &SGraphNodeAnimState::OnVerifyNameTextChanged)
+									.OnTextCommitted(this, &SGraphNodeAnimState::OnNameTextCommited)
+									.IsReadOnly( this, &SGraphNodeAnimState::IsNameReadOnly )
+									.IsSelected(this, &SGraphNodeAnimState::IsSelectedExclusively)
+								]
+								+SVerticalBox::Slot()
+								.AutoHeight()
+								[
+									NodeTitle.ToSharedRef()
+								]
 							]
 						]
 					]
 				]
 			]
+			+SOverlay::Slot()
+			[
+				SNew(SBorder)
+				.BorderImage(SelectionBrush)
+				.Padding(0)
+				.Visibility_Lambda([this]()
+				{
+					TSharedPtr<SGraphPanel> OwnerPanel = OwnerGraphPanelPtr.Pin();
+					if (!OwnerPanel.IsValid())
+					{
+						return EVisibility::Hidden;
+					}
+
+					return OwnerPanel->SelectionManager.IsNodeSelected(GraphNode) ? EVisibility::HitTestInvisible : EVisibility::Hidden;
+				})
+			]
 		];
 
 	ErrorReporting = ErrorText;
@@ -345,26 +356,41 @@ void SGraphNodeAnimState::CreatePinWidgets()
 {
 	UAnimStateNodeBase* StateNode = CastChecked<UAnimStateNodeBase>(GraphNode);
 
-	UEdGraphPin* CurPin = StateNode->GetOutputPin();
-	if (!CurPin->bHidden)
+	UEdGraphPin* OutputPin = StateNode->GetOutputPin();
+	if (!OutputPin->bHidden)
 	{
-		TSharedPtr<SGraphPin> NewPin = SNew(SStateMachineOutputPin, CurPin);
-
+		TSharedPtr<SGraphPin> NewPin = SNew(SStateMachineOutputPin, OutputPin);
+		this->AddPin(NewPin.ToSharedRef());
+	}
+	UEdGraphPin* InputPin = StateNode->GetInputPin();
+	if (!InputPin->bHidden)
+	{
+		TSharedPtr<SGraphPin> NewPin = SNew(SStateMachineInputPin, InputPin);
 		this->AddPin(NewPin.ToSharedRef());
 	}
 }
 
 void SGraphNodeAnimState::AddPin(const TSharedRef<SGraphPin>& PinToAdd)
 {
-	PinToAdd->SetOwner( SharedThis(this) );
-	RightNodeBox->AddSlot()
+	PinToAdd->SetOwner(SharedThis(this));
+	PinOverlay->AddSlot()
 		.HAlign(HAlign_Fill)
 		.VAlign(VAlign_Fill)
-		.FillHeight(1.0f)
 		[
 			PinToAdd
 		];
-	OutputPins.Add(PinToAdd);
+
+	switch (PinToAdd->GetPinObj()->Direction)
+	{
+	case EGPD_Input:
+		InputPins.Add(PinToAdd);
+		break;
+	case EGPD_Output:
+		OutputPins.Add(PinToAdd);
+		break;
+	default:
+		break;
+	}
 }
 
 TSharedPtr<SToolTip> SGraphNodeAnimState::GetComplexTooltip()
@@ -407,6 +433,21 @@ const FSlateBrush* SGraphNodeAnimState::GetNameIcon() const
 	return FAppStyle::GetBrush( TEXT("Graph.StateNode.Icon") );
 }
 
+FSlateColor SGraphNodeAnimState::GetIconColor() const
+{
+	return FLinearColor(0.8f, 0.8f, 0.8f);
+}
+
+const FSlateBrush* SGraphNodeAnimState::GetShadowBrush(bool bSelected) const
+{
+	if (GraphNode->IsA<UAnimStateAliasNode>())
+	{
+		return FAppStyle::GetBrush(TEXT("Graph.AnimAliasNode.Shadow"));
+	}
+
+	return FAppStyle::GetBrush(TEXT("Graph.Node.Shadow"));
+}
+
 /////////////////////////////////////////////////////
 // SGraphNodeAnimConduit
 
@@ -438,4 +479,9 @@ const FSlateBrush* SGraphNodeAnimConduit::GetNameIcon() const
 	return FAppStyle::GetBrush( TEXT("Graph.ConduitNode.Icon") );
 }
 
+FSlateColor SGraphNodeAnimConduit::GetIconColor() const
+{
+	return FLinearColor(0.38f, 0.45f, 0.21f);
+}
+
 #undef LOCTEXT_NAMESPACE
\ No newline at end of file
diff --git a/Engine/Source/Editor/AnimationBlueprintEditor/Private/AnimationStateNodes/SGraphNodeAnimState.h b/Engine/Source/Editor/AnimationBlueprintEditor/Private/AnimationStateNodes/SGraphNodeAnimState.h
--- a/Engine/Source/Editor/AnimationBlueprintEditor/Private/AnimationStateNodes/SGraphNodeAnimState.h
+++ b/Engine/Source/Editor/AnimationBlueprintEditor/Private/AnimationStateNodes/SGraphNodeAnimState.h
@@ -38,6 +38,7 @@ public:
 	virtual void CreatePinWidgets() override;
 	virtual void AddPin(const TSharedRef<SGraphPin>& PinToAdd) override;
 	virtual TSharedPtr<SToolTip> GetComplexTooltip() override;
+	virtual const FSlateBrush* GetShadowBrush(bool bSelected) const override;
 	// End of SGraphNode interface
 
 	// SWidget interface
@@ -45,13 +46,15 @@ public:
 	void OnMouseLeave(const FPointerEvent& MouseEvent) override;
 	// End of SWidget interface
 
-	static void GetStateInfoPopup(UEdGraphNode* GraphNode, TArray<FGraphInformationPopupInfo>& Popups);
+	void GetStateInfoPopup(TArray<FGraphInformationPopupInfo>& Popups) const;
 protected:
 	FSlateColor GetBorderBackgroundColor() const;
 	virtual FSlateColor GetBorderBackgroundColor_Internal(FLinearColor InactiveStateColor, FLinearColor ActiveStateColorDim, FLinearColor ActiveStateColorBright) const;
-
 	virtual FText GetPreviewCornerText() const;
 	virtual const FSlateBrush* GetNameIcon() const;
+	virtual FSlateColor GetIconColor() const;
+
+	TSharedPtr<SOverlay> PinOverlay;
 };
 
 
@@ -68,7 +71,7 @@ public:
 	// End of SNodePanel::SNode interface
 protected:
 	virtual FSlateColor GetBorderBackgroundColor_Internal(FLinearColor InactiveStateColor, FLinearColor ActiveStateColorDim, FLinearColor ActiveStateColorBright) const override;
-
 	virtual FText GetPreviewCornerText() const override;
 	virtual const FSlateBrush* GetNameIcon() const override;
+	virtual FSlateColor GetIconColor() const override;
 };
diff --git a/Engine/Source/Editor/AnimationBlueprintEditor/Private/AnimationStateNodes/SGraphNodeAnimStateAlias.cpp b/Engine/Source/Editor/AnimationBlueprintEditor/Private/AnimationStateNodes/SGraphNodeAnimStateAlias.cpp
--- a/Engine/Source/Editor/AnimationBlueprintEditor/Private/AnimationStateNodes/SGraphNodeAnimStateAlias.cpp
+++ b/Engine/Source/Editor/AnimationBlueprintEditor/Private/AnimationStateNodes/SGraphNodeAnimStateAlias.cpp
@@ -117,6 +117,11 @@ const FSlateBrush* SGraphNodeAnimStateAlias::GetNameIcon() const
 	return FAppStyle::GetBrush(TEXT("Graph.AliasNode.Icon"));
 }
 
+FSlateColor SGraphNodeAnimStateAlias::GetIconColor() const
+{
+	return FLinearColor(0.45f, 0.33f, 0.37f);
+}
+
 TSharedPtr<SToolTip> SGraphNodeAnimStateAlias::GetComplexTooltip()
 {
 	return nullptr;
diff --git a/Engine/Source/Editor/AnimationBlueprintEditor/Private/AnimationStateNodes/SGraphNodeAnimStateAlias.h b/Engine/Source/Editor/AnimationBlueprintEditor/Private/AnimationStateNodes/SGraphNodeAnimStateAlias.h
--- a/Engine/Source/Editor/AnimationBlueprintEditor/Private/AnimationStateNodes/SGraphNodeAnimStateAlias.h
+++ b/Engine/Source/Editor/AnimationBlueprintEditor/Private/AnimationStateNodes/SGraphNodeAnimStateAlias.h
@@ -34,7 +34,7 @@ public:
 
 protected:
 	virtual FSlateColor GetBorderBackgroundColor_Internal(FLinearColor InactiveStateColor, FLinearColor ActiveStateColorDim, FLinearColor ActiveStateColorBright) const;
-
 	virtual FText GetPreviewCornerText() const override;
 	virtual const FSlateBrush* GetNameIcon() const override;
+	virtual FSlateColor GetIconColor() const override;
 };
diff --git a/Engine/Source/Editor/AnimationBlueprintEditor/Private/AnimationStateNodes/SGraphNodeAnimStateEntry.cpp b/Engine/Source/Editor/AnimationBlueprintEditor/Private/AnimationStateNodes/SGraphNodeAnimStateEntry.cpp
--- a/Engine/Source/Editor/AnimationBlueprintEditor/Private/AnimationStateNodes/SGraphNodeAnimStateEntry.cpp
+++ b/Engine/Source/Editor/AnimationBlueprintEditor/Private/AnimationStateNodes/SGraphNodeAnimStateEntry.cpp
@@ -4,6 +4,7 @@
 #include "AnimationStateNodes/SGraphNodeAnimStateEntry.h"
 
 #include "AnimStateEntryNode.h"
+#include "SGraphPanel.h"
 #include "Delegates/Delegate.h"
 #include "GenericPlatform/ICursor.h"
 #include "HAL/PlatformCrt.h"
@@ -21,6 +22,10 @@
 #include "Widgets/Layout/SBorder.h"
 #include "Widgets/SBoxPanel.h"
 #include "Widgets/SOverlay.h"
+#include "SStateMachineOutputPin.h"
+#include "Widgets/Images/SImage.h"
+
+#define LOCTEXT_NAMESPACE "SGraphNodeAnimStateEntry"
 
 /////////////////////////////////////////////////////
 // SGraphNodeAnimStateEntry
@@ -58,6 +63,9 @@ void SGraphNodeAnimStateEntry::UpdateGraphNode()
 	LeftNodeBox.Reset();
 
 
+	// Set ourselves as SelfHitTestInvisible so we dont end up with a 1-unit border where the node is selected rather than the pin
+	SetVisibility(EVisibility::SelfHitTestInvisible);
+
 	FLinearColor TitleShadowColor(0.6f, 0.6f, 0.6f);
 
 	this->ContentScale.Bind( this, &SGraphNode::GetContentScale );
@@ -65,27 +73,94 @@ void SGraphNodeAnimStateEntry::UpdateGraphNode()
 		.HAlign(HAlign_Center)
 		.VAlign(VAlign_Center)
 		[
-			SNew(SBorder)
-			.BorderImage( FAppStyle::GetBrush( "Graph.StateNode.Body" ) )
-			.Padding(0)
-			.BorderBackgroundColor( this, &SGraphNodeAnimStateEntry::GetBorderBackgroundColor )
+			SNew(SOverlay)
+			.Visibility(EVisibility::SelfHitTestInvisible)
+			+SOverlay::Slot()
+			.Padding(1.0f)
 			[
-				SNew(SOverlay)
-
-				// PIN AREA
-				+SOverlay::Slot()
-				.HAlign(HAlign_Fill)
-				.VAlign(VAlign_Fill)
-				.Padding(10.0f)
+				SNew(SBorder)
+				.BorderImage( FAppStyle::GetBrush( "Graph.AnimConduitNode.Body" ) )
+				.Padding(0)
+				.BorderBackgroundColor( this, &SGraphNodeAnimStateEntry::GetBorderBackgroundColor )
 				[
-					SAssignNew(RightNodeBox, SVerticalBox)
+					SNew(SOverlay)
+
+					// PIN AREA
+					+SOverlay::Slot()
+					.HAlign(HAlign_Fill)
+					.VAlign(VAlign_Fill)
+					[
+						SAssignNew(RightNodeBox, SVerticalBox)
+					]
+					// STATE NAME AREA
+					+SOverlay::Slot()
+					.HAlign(HAlign_Center)
+					.VAlign(VAlign_Center)
+					.Padding(12.0f)
+					[
+						SNew(SBorder)
+						.BorderImage( FAppStyle::GetBrush( "Graph.AnimConduitNode.ColorSpill" ) )
+						.BorderBackgroundColor( TitleShadowColor )
+						.HAlign(HAlign_Center)
+						.VAlign(VAlign_Center)
+						.Padding(FMargin(2.0f, 2.0f, 2.0f, 2.0f))
+						[
+							SNew(SHorizontalBox)
+							+SHorizontalBox::Slot()
+							.AutoWidth()
+							.VAlign(VAlign_Center)
+							.Padding(FMargin(8.0f, 0.0f, 4.0f, 0.0f))
+							[
+								SNew(SImage)
+								.Image(FAppStyle::GetBrush("Graph.EntryNode.Icon"))
+							]
+							+SHorizontalBox::Slot()
+							.AutoWidth()
+							.VAlign(VAlign_Center)
+							.Padding(FMargin(2.0f, 2.5f, 8.0f, 2.5f))
+							[
+								SNew(STextBlock)
+								.Text(LOCTEXT("EntryName", "Entry"))
+								.TextStyle(FAppStyle::Get(), "Graph.AnimStateNode.NodeTitle")
+							]
+						]
+					]
 				]
-			]			
+			]
+			+SOverlay::Slot()
+			[
+				SNew(SBorder)
+				.BorderImage(FAppStyle::GetBrush( "Graph.AnimConduitNode.Selection"))
+				.Padding(0)
+				.Visibility_Lambda([this]()
+				{
+					TSharedPtr<SGraphPanel> OwnerPanel = OwnerGraphPanelPtr.Pin();
+					if (!OwnerPanel.IsValid())
+					{
+						return EVisibility::Hidden;
+					}
+
+					return OwnerPanel->SelectionManager.IsNodeSelected(GraphNode) ? EVisibility::HitTestInvisible : EVisibility::Hidden;
+				})
+			]
 		];
 
 	CreatePinWidgets();
 }
 
+void SGraphNodeAnimStateEntry::CreatePinWidgets()
+{
+	UAnimStateEntryNode* StateNode = CastChecked<UAnimStateEntryNode>(GraphNode);
+
+	UEdGraphPin* CurPin = StateNode->GetOutputPin();
+	if (!CurPin->bHidden)
+	{
+		TSharedPtr<SGraphPin> NewPin = SNew(SStateMachineOutputPin, CurPin);
+
+		this->AddPin(NewPin.ToSharedRef());
+	}
+}
+
 void SGraphNodeAnimStateEntry::AddPin(const TSharedRef<SGraphPin>& PinToAdd)
 {
 	PinToAdd->SetOwner( SharedThis(this) );
@@ -101,5 +176,12 @@ void SGraphNodeAnimStateEntry::AddPin(const TSharedRef<SGraphPin>& PinToAdd)
 
 FText SGraphNodeAnimStateEntry::GetPreviewCornerText() const
 {
-	return NSLOCTEXT("SGraphNodeAnimStateEntry", "CornerTextDescription", "Entry point for state machine");
+	return LOCTEXT("CornerTextDescription", "Entry point for state machine");
+}
+
+const FSlateBrush* SGraphNodeAnimStateEntry::GetShadowBrush(bool bSelected) const
+{
+	return FAppStyle::GetBrush(TEXT("Graph.Node.Shadow"));
 }
+
+#undef LOCTEXT_NAMESPACE
\ No newline at end of file
diff --git a/Engine/Source/Editor/AnimationBlueprintEditor/Private/AnimationStateNodes/SGraphNodeAnimStateEntry.h b/Engine/Source/Editor/AnimationBlueprintEditor/Private/AnimationStateNodes/SGraphNodeAnimStateEntry.h
--- a/Engine/Source/Editor/AnimationBlueprintEditor/Private/AnimationStateNodes/SGraphNodeAnimStateEntry.h
+++ b/Engine/Source/Editor/AnimationBlueprintEditor/Private/AnimationStateNodes/SGraphNodeAnimStateEntry.h
@@ -29,7 +29,8 @@ public:
 	// SGraphNode interface
 	virtual void UpdateGraphNode() override;
 	virtual void AddPin(const TSharedRef<SGraphPin>& PinToAdd) override;
-	
+	virtual const FSlateBrush* GetShadowBrush(bool bSelected) const override;
+	virtual void CreatePinWidgets() override;
 	// End of SGraphNode interface
 
 	
diff --git a/Engine/Source/Editor/AnimationBlueprintEditor/Private/AnimationStateNodes/SGraphNodeAnimTransition.cpp b/Engine/Source/Editor/AnimationBlueprintEditor/Private/AnimationStateNodes/SGraphNodeAnimTransition.cpp
--- a/Engine/Source/Editor/AnimationBlueprintEditor/Private/AnimationStateNodes/SGraphNodeAnimTransition.cpp
+++ b/Engine/Source/Editor/AnimationBlueprintEditor/Private/AnimationStateNodes/SGraphNodeAnimTransition.cpp
@@ -42,6 +42,7 @@
 #include "Widgets/SBoxPanel.h"
 #include "Widgets/SOverlay.h"
 #include "Widgets/SToolTip.h"
+#include "Widgets/Layout/SBox.h"
 #include "Widgets/Text/STextBlock.h"
 
 class SWidget;
@@ -239,6 +240,17 @@ void SGraphNodeAnimTransition::UpdateGraphNode()
 	RightNodeBox.Reset();
 	LeftNodeBox.Reset();
 
+	TSharedRef<SImage> DirectionImage = SNew(SImage)
+		.Image(this, &SGraphNodeAnimTransition::GetTransitionIconImage)
+		.ColorAndOpacity(FStyleColors::Background);
+
+	DirectionImage->SetRenderTransform(MakeAttributeLambda([this]()-> TOptional<FSlateRenderTransform>
+	{
+		UAnimStateTransitionNode* TransNode = CastChecked<UAnimStateTransitionNode>(GraphNode);
+		return FSlateRenderTransform(FQuat2D(TransNode->CachedRotation));
+	}));
+	DirectionImage->SetRenderTransformPivot(FVector2D(0.5f, 0.5f));
+
 	this->ContentScale.Bind( this, &SGraphNode::GetContentScale );
 	this->GetOrAddSlot( ENodeZone::Center )
 		.HAlign(HAlign_Center)
@@ -246,15 +258,35 @@ void SGraphNodeAnimTransition::UpdateGraphNode()
 		[
 			SNew(SOverlay)
 			+SOverlay::Slot()
+			.Padding(2.0f)
 			[
 				SNew(SImage)
-				.Image( FAppStyle::GetBrush("Graph.TransitionNode.ColorSpill") )
-				.ColorAndOpacity( this, &SGraphNodeAnimTransition::GetTransitionColor )
+				.Image(FAppStyle::GetBrush("Graph.AnimTransitionNode.ColorSpill"))
+				.ColorAndOpacity(this, &SGraphNodeAnimTransition::GetTransitionColor)
 			]
 			+SOverlay::Slot()
 			[
-				SNew(SImage)
-				.Image( this, &SGraphNodeAnimTransition::GetTransitionIconImage )
+				SNew(SBox)
+				.Padding(4.0f)
+				[
+					DirectionImage
+				]
+			]
+			+SOverlay::Slot()
+			[
+				SNew(SBorder)
+				.BorderImage(FAppStyle::GetBrush("Graph.AnimTransitionNode.Selection"))
+				.Padding(0)
+				.Visibility_Lambda([this]()
+				{
+					TSharedPtr<SGraphPanel> OwnerPanel = OwnerGraphPanelPtr.Pin();
+					if (!OwnerPanel.IsValid())
+					{
+						return EVisibility::Hidden;
+					}
+
+					return OwnerPanel->SelectionManager.IsNodeSelected(GraphNode) ? EVisibility::HitTestInvisible : EVisibility::Hidden;
+				})
 			]
 		];
 }
@@ -301,12 +333,12 @@ FText SGraphNodeAnimTransition::GetPreviewCornerText(bool bReverse) const
 	return Result;
 }
 
-FLinearColor SGraphNodeAnimTransition::StaticGetTransitionColor(UAnimStateTransitionNode* TransNode, bool bIsHovered)
+FSlateColor SGraphNodeAnimTransition::StaticGetTransitionColor(UAnimStateTransitionNode* TransNode, bool bIsHovered)
 {
-	//@TODO: Make configurable by styling
-	const FLinearColor ActiveColor(1.0f, 0.4f, 0.3f, 1.0f);
-	const FLinearColor HoverColor(0.724f, 0.256f, 0.0f, 1.0f);
-	FLinearColor BaseColor(0.9f, 0.9f, 0.9f, 1.0f);
+	FLinearColor ActiveColor = FStyleColors::AccentOrange.GetSpecifiedColor().Desaturate(0.25f);
+	ActiveColor.A = 1.0f;
+	const FSlateColor HoverColor = FStyleColors::AccentOrange;
+	FSlateColor BaseColor = FStyleColors::Foreground;
 
 	// Display various types of debug data
 	UAnimBlueprint* AnimBlueprint = Cast<UAnimBlueprint>(FBlueprintEditorUtils::FindBlueprintForNodeChecked(TransNode));
@@ -361,12 +393,17 @@ FLinearColor SGraphNodeAnimTransition::StaticGetTransitionColor(UAnimStateTransi
 	// If shared transition, show different color
 	if (TransNode->bSharedRules)
 	{
-		BaseColor = TransNode->SharedColor;
+		// Override alpha of shared colors as some have old 0.25f alpha serialized
+		FLinearColor SharedColor(TransNode->SharedColor);
+		SharedColor.A = 1.0f;
+		BaseColor = SharedColor;
 	}
 
 	if (TransNode->bDisabled)
 	{
-		BaseColor = FLinearColor(0.125, 0.125, 0.125, 0.5);
+		FLinearColor DisabledColor(BaseColor.GetSpecifiedColor());
+		DisabledColor.A *= 0.5f;
+		BaseColor = DisabledColor;
 	}
 
 	return bIsHovered ? HoverColor : BaseColor;
@@ -398,8 +435,8 @@ const FSlateBrush* SGraphNodeAnimTransition::GetTransitionIconImage() const
 {
 	UAnimStateTransitionNode* TransNode = CastChecked<UAnimStateTransitionNode>(GraphNode);
 	return (TransNode->LogicType == ETransitionLogicType::TLT_Inertialization)
-		? FAppStyle::GetBrush("Graph.TransitionNode.Icon_Inertialization")
-		: FAppStyle::GetBrush("Graph.TransitionNode.Icon");
+		? FAppStyle::GetBrush("Graph.AnimTransitionNode.Icon_Inertialization")
+		: FAppStyle::GetBrush("Graph.AnimTransitionNode.Icon");
 }
 
 FString SGraphNodeAnimTransition::GetCurrentDuration() const
@@ -431,6 +468,12 @@ void SGraphNodeAnimTransition::OnMouseLeave(const FPointerEvent& MouseEvent)
 	SGraphNode::OnMouseLeave(MouseEvent);
 }
 
+int32 SGraphNodeAnimTransition::OnPaint(const FPaintArgs& Args, const FGeometry& AllottedGeometry, const FSlateRect& MyCullingRect, FSlateWindowElementList& OutDrawElements, int32 LayerId, const FWidgetStyle& InWidgetStyle, bool bParentEnabled) const
+{
+	// Boost layer ID to sort in front of states
+	return SGraphNode::OnPaint(Args, AllottedGeometry, MyCullingRect, OutDrawElements, LayerId + 100, InWidgetStyle, bParentEnabled);
+}
+
 void SGraphNodeAnimTransition::PositionBetweenTwoNodesWithOffset(const FGeometry& StartGeom, const FGeometry& EndGeom, int32 NodeIndex, int32 MaxNodes) const
 {
 	// Get a reasonable seed point (halfway between the boxes)
@@ -443,7 +486,7 @@ void SGraphNodeAnimTransition::PositionBetweenTwoNodesWithOffset(const FGeometry
 	const FVector2D EndAnchorPoint = FGeometryHelper::FindClosestPointOnGeom(EndGeom, SeedPoint);
 
 	// Position ourselves halfway along the connecting line between the nodes, elevated away perpendicular to the direction of the line
-	const float Height = 30.0f;
+	const float Height = 24.0f;
 
 	const FVector2D DesiredNodeSize = GetDesiredSize();
 
@@ -464,17 +507,25 @@ void SGraphNodeAnimTransition::PositionBetweenTwoNodesWithOffset(const FGeometry
 	// MultiNodeOffset: the offset where 0 is the centre of the transition, -1 is 1 <size of node>
 	// towards the PrevStateNode and +1 is 1 <size of node> towards the NextStateNode.
 
-	const float MutliNodeSpace = 0.2f; // Space between multiple transition nodes (in units of <size of node> )
+	const float MutliNodeSpace = 0.0f; // Space between multiple transition nodes (in units of <size of node> )
 	const float MultiNodeStep = (1.f + MutliNodeSpace); //Step between node centres (Size of node + size of node spacer)
 
 	const float MultiNodeStart = -((MaxNodes - 1) * MultiNodeStep) / 2.f;
 	const float MultiNodeOffset = MultiNodeStart + (NodeIndex * MultiNodeStep);
 
 	// Now we need to adjust the new center by the node size, zoom factor and multi node offset
-	const FVector2D NewCorner = NewCenter - (0.5f * DesiredNodeSize) + (DeltaNormal * MultiNodeOffset * DesiredNodeSize.Size());
+	const FVector2D NewCorner = NewCenter - (0.5f * DesiredNodeSize) + (DeltaNormal * MultiNodeOffset * DesiredNodeSize.X);
 
 	GraphNode->NodePosX = static_cast<int32>(NewCorner.X);
 	GraphNode->NodePosY = static_cast<int32>(NewCorner.Y);
+
+	UAnimStateTransitionNode* TransNode = CastChecked<UAnimStateTransitionNode>(GraphNode);
+	TransNode->CachedRotation = DeltaNormal;
+}
+
+const FSlateBrush* SGraphNodeAnimTransition::GetShadowBrush(bool bSelected) const
+{
+	return FAppStyle::GetBrush(TEXT("Graph.AnimTransitionNode.Shadow"));
 }
 
 /////////////////////////////////////////////////////
diff --git a/Engine/Source/Editor/AnimationBlueprintEditor/Private/AnimationStateNodes/SGraphNodeAnimTransition.h b/Engine/Source/Editor/AnimationBlueprintEditor/Private/AnimationStateNodes/SGraphNodeAnimTransition.h
--- a/Engine/Source/Editor/AnimationBlueprintEditor/Private/AnimationStateNodes/SGraphNodeAnimTransition.h
+++ b/Engine/Source/Editor/AnimationBlueprintEditor/Private/AnimationStateNodes/SGraphNodeAnimTransition.h
@@ -42,17 +42,19 @@ public:
 	// SGraphNode interface
 	virtual void UpdateGraphNode() override;
 	virtual TSharedPtr<SToolTip> GetComplexTooltip() override;
+	virtual const FSlateBrush* GetShadowBrush(bool bSelected) const override;
 	// End of SGraphNode interface
 
 	// SWidget interface
-	void OnMouseEnter(const FGeometry& MyGeometry, const FPointerEvent& MouseEvent) override;
-	void OnMouseLeave(const FPointerEvent& MouseEvent) override;
+	virtual void OnMouseEnter(const FGeometry& MyGeometry, const FPointerEvent& MouseEvent) override;
+	virtual void OnMouseLeave(const FPointerEvent& MouseEvent) override;
+	virtual int32 OnPaint(const FPaintArgs& Args, const FGeometry& AllottedGeometry, const FSlateRect& MyCullingRect, FSlateWindowElementList& OutDrawElements, int32 LayerId, const FWidgetStyle& InWidgetStyle, bool bParentEnabled) const;
 	// End of SWidget interface
 
 	// Calculate position for multiple nodes to be placed between a start and end point, by providing this nodes index and max expected nodes 
 	void PositionBetweenTwoNodesWithOffset(const FGeometry& StartGeom, const FGeometry& EndGeom, int32 NodeIndex, int32 MaxNodes) const;
 
-	static FLinearColor StaticGetTransitionColor(UAnimStateTransitionNode* TransNode, bool bIsHovered);
+	static FSlateColor StaticGetTransitionColor(UAnimStateTransitionNode* TransNode, bool bIsHovered);
 
 	static bool IsTransitionActive(int32 TransitionIndex, class UAnimBlueprintGeneratedClass& AnimClass, class UAnimationStateMachineGraph& StateMachineGraph, class UAnimInstance& AnimInstance);
 private:
diff --git a/Engine/Source/Editor/AnimationBlueprintEditor/Private/AnimationStateNodes/SStateMachineInputPin.cpp b/Engine/Source/Editor/AnimationBlueprintEditor/Private/AnimationStateNodes/SStateMachineInputPin.cpp
new file mode 100644
--- /dev/null
+++ b/Engine/Source/Editor/AnimationBlueprintEditor/Private/AnimationStateNodes/SStateMachineInputPin.cpp
@@ -0,0 +1,27 @@
+﻿// Copyright Epic Games, Inc. All Rights Reserved.
+
+#include "SStateMachineInputPin.h"
+
+#include "SGraphPanel.h"
+
+void SStateMachineInputPin::Construct(const FArguments& InArgs, UEdGraphPin* InPin)
+{
+	this->SetCursor(EMouseCursor::Default);
+
+	bShowLabel = false;
+
+	GraphPinObj = InPin;
+	check(GraphPinObj != nullptr);
+
+	SetVisibility(EVisibility::HitTestInvisible);
+
+	// Set up a brush-less border so this pin is effectively hidden
+	SBorder::Construct( SBorder::FArguments()
+		.BorderImage(nullptr)
+	);
+}
+
+TSharedRef<SWidget>	SStateMachineInputPin::GetDefaultValueWidget()
+{
+	return SNullWidget::NullWidget;
+}
diff --git a/Engine/Source/Editor/AnimationBlueprintEditor/Private/AnimationStateNodes/SStateMachineInputPin.h b/Engine/Source/Editor/AnimationBlueprintEditor/Private/AnimationStateNodes/SStateMachineInputPin.h
new file mode 100644
--- /dev/null
+++ b/Engine/Source/Editor/AnimationBlueprintEditor/Private/AnimationStateNodes/SStateMachineInputPin.h
@@ -0,0 +1,19 @@
+﻿// Copyright Epic Games, Inc. All Rights Reserved.
+
+#pragma once
+
+#include "SGraphPin.h"
+
+class SStateMachineInputPin : public SGraphPin
+{
+public:
+	SLATE_BEGIN_ARGS(SStateMachineInputPin){}
+	SLATE_END_ARGS()
+
+	void Construct(const FArguments& InArgs, UEdGraphPin* InPin);
+protected:
+	// Begin SGraphPin interface
+	virtual TSharedRef<SWidget>	GetDefaultValueWidget() override;
+	// End SGraphPin interface
+};
+
diff --git a/Engine/Source/Editor/AnimationBlueprintEditor/Private/AnimationStateNodes/SStateMachineOutputPin.cpp b/Engine/Source/Editor/AnimationBlueprintEditor/Private/AnimationStateNodes/SStateMachineOutputPin.cpp
new file mode 100644
--- /dev/null
+++ b/Engine/Source/Editor/AnimationBlueprintEditor/Private/AnimationStateNodes/SStateMachineOutputPin.cpp
@@ -0,0 +1,75 @@
+﻿// Copyright Epic Games, Inc. All Rights Reserved.
+
+#include "SStateMachineOutputPin.h"
+
+#include "AnimStateAliasNode.h"
+#include "AnimStateConduitNode.h"
+#include "AnimStateEntryNode.h"
+#include "AnimStateNode.h"
+#include "SGraphPanel.h"
+
+void SStateMachineOutputPin::Construct(const FArguments& InArgs, UEdGraphPin* InPin)
+{
+	this->SetCursor( EMouseCursor::Default );
+
+	bShowLabel = true;
+
+	GraphPinObj = InPin;
+	check(GraphPinObj != NULL);
+
+	const UEdGraphSchema* Schema = GraphPinObj->GetSchema();
+	check(Schema);
+
+	// Set up a hover for pins that is tinted the color of the pin.
+	SBorder::Construct( SBorder::FArguments()
+		.BorderImage( this, &SStateMachineOutputPin::GetPinBorder )
+		.BorderBackgroundColor( this, &SStateMachineOutputPin::GetPinColor )
+		.OnMouseButtonDown( this, &SStateMachineOutputPin::OnPinMouseDown )
+		.Cursor( this, &SStateMachineOutputPin::GetPinCursor )
+	);
+}
+
+TSharedRef<SWidget>	SStateMachineOutputPin::GetDefaultValueWidget()
+{
+	return SNew(STextBlock);
+}
+
+const FSlateBrush* SStateMachineOutputPin::GetPinBorder() const
+{
+	UEdGraphNode* Node = GraphPinObj->GetOwningNode();
+	if (Node)
+	{
+		if (Node->IsA<UAnimStateAliasNode>())
+		{
+			return ( IsHovered() )
+				? FAppStyle::GetBrush( TEXT("Graph.AnimAliasNode.Pin.BackgroundHovered") )
+				: FAppStyle::GetBrush( TEXT("Graph.AnimStateNode.Pin.Background") );
+		}
+		else if (Node->IsA<UAnimStateConduitNode>() || Node->IsA<UAnimStateEntryNode>())
+		{
+			return ( IsHovered() )
+				? FAppStyle::GetBrush( TEXT("Graph.AnimConduitNode.Pin.BackgroundHovered") )
+				: FAppStyle::GetBrush( TEXT("Graph.AnimStateNode.Pin.Background") );
+		}
+	}
+
+	bool bAliasHovered = false;
+	if (UAnimStateNode* StateNode = Cast<UAnimStateNode>(Node))
+	{
+		const TSharedPtr<SGraphNode> OwnerNode = OwnerNodePtr.Pin();
+		if(OwnerNode.IsValid())
+		{
+			if (UAnimStateAliasNode* HoveredAlias = Cast<UAnimStateAliasNode>(OwnerNode->GetOwnerPanel()->GetCurrentHoveredNode()))
+			{
+				if (HoveredAlias->bGlobalAlias || HoveredAlias->GetAliasedStates().Contains(StateNode))
+				{
+					bAliasHovered = true;
+				}
+			}
+		}
+	}
+
+	return ( IsHovered() || bAliasHovered )
+		? FAppStyle::GetBrush( TEXT("Graph.AnimStateNode.Pin.BackgroundHovered") )
+		: FAppStyle::GetBrush( TEXT("Graph.AnimStateNode.Pin.Background") );
+}
\ No newline at end of file
diff --git a/Engine/Source/Editor/AnimationBlueprintEditor/Private/AnimationStateNodes/SStateMachineOutputPin.h b/Engine/Source/Editor/AnimationBlueprintEditor/Private/AnimationStateNodes/SStateMachineOutputPin.h
new file mode 100644
--- /dev/null
+++ b/Engine/Source/Editor/AnimationBlueprintEditor/Private/AnimationStateNodes/SStateMachineOutputPin.h
@@ -0,0 +1,21 @@
+﻿// Copyright Epic Games, Inc. All Rights Reserved.
+
+#pragma once
+
+#include "SGraphPin.h"
+
+class SStateMachineOutputPin : public SGraphPin
+{
+public:
+	SLATE_BEGIN_ARGS(SStateMachineOutputPin){}
+	SLATE_END_ARGS()
+
+	void Construct(const FArguments& InArgs, UEdGraphPin* InPin);
+protected:
+	// Begin SGraphPin interface
+	virtual TSharedRef<SWidget>	GetDefaultValueWidget() override;
+	// End SGraphPin interface
+
+	const FSlateBrush* GetPinBorder() const;
+};
+
diff --git a/Engine/Source/Editor/AnimationBlueprintEditor/Private/StateMachineConnectionDrawingPolicy.cpp b/Engine/Source/Editor/AnimationBlueprintEditor/Private/StateMachineConnectionDrawingPolicy.cpp
--- a/Engine/Source/Editor/AnimationBlueprintEditor/Private/StateMachineConnectionDrawingPolicy.cpp
+++ b/Engine/Source/Editor/AnimationBlueprintEditor/Private/StateMachineConnectionDrawingPolicy.cpp
@@ -3,6 +3,7 @@
 #include "StateMachineConnectionDrawingPolicy.h"
 
 #include "AnimStateEntryNode.h"
+#include "AnimStateNode.h"
 #include "AnimStateNodeBase.h"
 #include "AnimStateTransitionNode.h"
 #include "AnimationStateNodes/SGraphNodeAnimTransition.h"
@@ -35,11 +36,11 @@ FStateMachineConnectionDrawingPolicy::FStateMachineConnectionDrawingPolicy(int32
 	: FConnectionDrawingPolicy(InBackLayerID, InFrontLayerID, ZoomFactor, InClippingRect, InDrawElements)
 	, GraphObj(InGraphObj)
 {
+	ArrowImage = FAppStyle::GetBrush( TEXT("Graph.AnimStateNode.ConnectionArrow") );
 }
 
 void FStateMachineConnectionDrawingPolicy::DetermineWiringStyle(UEdGraphPin* OutputPin, UEdGraphPin* InputPin, /*inout*/ FConnectionParams& Params)
 {
-	Params.bUserFlag2 = true;
 	Params.AssociatedPin1 = OutputPin;
 	Params.AssociatedPin2 = InputPin;
 	Params.WireThickness = 1.5f;
@@ -49,42 +50,21 @@ void FStateMachineConnectionDrawingPolicy::DetermineWiringStyle(UEdGraphPin* Out
 		if (UAnimStateTransitionNode* TransNode = Cast<UAnimStateTransitionNode>(InputPin->GetOwningNode()))
 		{
 			const bool IsInputPinHovered = HoveredPins.Contains(InputPin);
-			Params.WireColor = SGraphNodeAnimTransition::StaticGetTransitionColor(TransNode, IsInputPinHovered);
+			Params.WireColor = SGraphNodeAnimTransition::StaticGetTransitionColor(TransNode, IsInputPinHovered).GetSpecifiedColor();
 			Params.bUserFlag1 = TransNode->Bidirectional;
 		}
 	}
 
 	const bool bDeemphasizeUnhoveredPins = HoveredPins.Num() > 0;
-	if (bDeemphasizeUnhoveredPins)
+	if (bDeemphasizeUnhoveredPins && !Params.bUserFlag2)
 	{
 		ApplyHoverDeemphasis(OutputPin, InputPin, /*inout*/ Params.WireThickness, /*inout*/ Params.WireColor);
 	}
 
 	// Make the transition that is currently relinked, semi-transparent.
-	for (const FRelinkConnection& Connection : RelinkConnections)
+	if (Params.bUserFlag2)
 	{
-		if (InputPin == nullptr || OutputPin == nullptr)
-		{
-			continue;
-		}
-		const FGraphPinHandle SourcePinHandle = Connection.SourcePin;
-		const FGraphPinHandle TargetPinHandle = Connection.TargetPin;
-
-		// Skip all transitions that don't start at the node our dragged and relink transition starts from
-		if (OutputPin->GetOwningNode()->NodeGuid == SourcePinHandle.NodeGuid)
-		{
-			// Safety check to verify if the node is a transition node
-			if (UAnimStateTransitionNode* TransitonNode = Cast<UAnimStateTransitionNode>(InputPin->GetOwningNode()))
-			{
-				if (UEdGraphPin* TransitionOutputPin = TransitonNode->GetOutputPin())
-				{
-					if (TargetPinHandle.NodeGuid == TransitionOutputPin->GetOwningNode()->NodeGuid)
-					{
-						Params.WireColor.A *= 0.2f;
-					}
-				}
-			}
-		}
+		Params.WireColor.A = 0.5f;
 	}
 }
 
@@ -120,16 +100,6 @@ void FStateMachineConnectionDrawingPolicy::DetermineLinkGeometry(
 			}
 		}
 	}
-	else
-	{
-		StartWidgetGeometry = PinGeometries->Find(OutputPinWidget);
-
-		if (TSharedPtr<SGraphPin>* pTargetWidget = PinToPinWidgetMap.Find(InputPin))
-		{
-			TSharedRef<SGraphPin> InputWidget = (*pTargetWidget).ToSharedRef();
-			EndWidgetGeometry = PinGeometries->Find(InputWidget);
-		}
-	}
 }
 
 void FStateMachineConnectionDrawingPolicy::Draw(TMap<TSharedRef<SWidget>, FArrangedWidget>& InPinGeometries, FArrangedChildren& ArrangedNodes)
@@ -150,13 +120,17 @@ void FStateMachineConnectionDrawingPolicy::Draw(TMap<TSharedRef<SWidget>, FArran
 void FStateMachineConnectionDrawingPolicy::DrawPreviewConnector(const FGeometry& PinGeometry, const FVector2f& StartPoint, const FVector2f& EndPoint, UEdGraphPin* Pin)
 {
 	FConnectionParams Params;
-	DetermineWiringStyle(Pin, nullptr, /*inout*/ Params);
+	Params.bUserFlag2 = true;	// bUserFlag2 is used to indicate whether the drawn arrow is a preview transition (the temporary transition when creating or relinking).
+	UEdGraphPin* InputPin = Pin->Direction == EGPD_Input ? Pin : nullptr;
+	UEdGraphPin* OutputPin = Pin->Direction == EGPD_Output ? Pin : nullptr;
+	DetermineWiringStyle(InputPin, OutputPin, /*inout*/ Params);
 
-	const FVector2f SeedPoint = EndPoint;
-	const FVector2f AdjustedStartPoint = FGeometryHelper::FindClosestPointOnGeom(PinGeometry, SeedPoint);
+	const FVector2f SeedPoint = Pin->Direction == EGPD_Input ? StartPoint : EndPoint;
+	const FVector2f AdjustedSeedPoint = FGeometryHelper::FindClosestPointOnGeom(PinGeometry, SeedPoint);
+	const FVector2f AdjustedStartPoint = Pin->Direction == EGPD_Input ? StartPoint : AdjustedSeedPoint;
+	const FVector2f AdjustedEndPoint = Pin->Direction == EGPD_Input ? AdjustedSeedPoint : EndPoint;
 
-	Params.bUserFlag2 = false; // bUserFlag2 is used to indicate whether the drawn arrow is a preview transition (the temporary transition when creating or relinking).
-	DrawSplineWithArrow(AdjustedStartPoint, EndPoint, Params);
+	DrawSplineWithArrow(AdjustedStartPoint, AdjustedEndPoint, Params);
 }
 
 
@@ -173,8 +147,7 @@ void FStateMachineConnectionDrawingPolicy::DrawSplineWithArrow(const FVector2f&
 
 void FStateMachineConnectionDrawingPolicy::Internal_DrawLineWithArrow(const FVector2f& StartAnchorPoint, const FVector2f& EndAnchorPoint, const FConnectionParams& Params)
 {
-	//@TODO: Should this be scaled by zoom factor?
-	const float LineSeparationAmount = 4.5f;
+	const float LineSeparationAmount = 6.0f * ZoomFactor;
 
 	const FVector2f DeltaPos = EndAnchorPoint - StartAnchorPoint;
 	const FVector2f UnitDelta = DeltaPos.GetSafeNormal();
@@ -187,60 +160,81 @@ void FStateMachineConnectionDrawingPolicy::Internal_DrawLineWithArrow(const FVec
 	const FVector2f EndPoint = EndAnchorPoint + DirectionBias - LengthBias;
 	FLinearColor ArrowHeadColor = Params.WireColor;
 
-	// Draw a line/spline
-	DrawConnection(WireLayerID, StartPoint, EndPoint, Params);
+	// Draw a line/spline.
+	// LengthBias * 0.8f is to ensure the line doesnt overlap the arrowhead glyph so they dont show up overlapping when semi-transparent
+	DrawConnection(WireLayerID, StartPoint, EndPoint - (LengthBias * 0.8f), Params);
 
 	const FVector2f ArrowDrawPos = EndPoint - ArrowRadius;
 	const double AngleInRadians = FMath::Atan2(DeltaPos.Y, DeltaPos.X);
 
 	// Draw the transition grab handles in case the mouse is hovering the transition
-	bool StartHovered = false;
-	bool EndHovered = false;
+	bool bStartHovered = false;
+	bool bEndHovered = false;
 	const FVector2f FVecMousePos = FVector2f(LocalMousePosition.X, LocalMousePosition.Y);
 	const FVector2f ClosestPoint = FMath::ClosestPointOnSegment2D(FVecMousePos, StartPoint, EndPoint);
-	if ((ClosestPoint - FVecMousePos).Length() < RelinkHandleHoverRadius)
+	if ((ClosestPoint - FVecMousePos).Length() < RelinkHandleHoverRadius * ZoomFactor)
 	{
-		StartHovered = (StartPoint - LocalMousePosition).Length() < RelinkHandleHoverRadius;
-		EndHovered = (EndPoint - LocalMousePosition).Length() < RelinkHandleHoverRadius;
-		FVector2f HoverIndicatorPosition = StartHovered ? StartPoint : EndPoint;
+		bStartHovered = (StartPoint - LocalMousePosition).Length() < RelinkHandleHoverRadius * ZoomFactor;
+		bEndHovered = (EndPoint - LocalMousePosition).Length() < RelinkHandleHoverRadius * ZoomFactor;
+		FVector2f HoverIndicatorPosition = bStartHovered ? StartPoint : EndPoint;
 
 		// Set the hovered pin results. This will be used by the SGraphPanel again.
 		const float SquaredDistToPin1 = (Params.AssociatedPin1 != nullptr) ? (StartPoint - LocalMousePosition).SizeSquared() : FLT_MAX;
 		const float SquaredDistToPin2 = (Params.AssociatedPin2 != nullptr) ? (EndPoint - LocalMousePosition).SizeSquared() : FLT_MAX;
-		if (EndHovered)
+		UEdGraphPin* Pin1 = Params.AssociatedPin1;
+		UEdGraphPin* Pin2 = Params.AssociatedPin2;
+		if (Pin2)
 		{
-			SplineOverlapResult = FGraphSplineOverlapResult(Params.AssociatedPin1, Params.AssociatedPin2, SquaredDistToPin2, SquaredDistToPin1, SquaredDistToPin2, true);
+			// Forward transition links to the destination state node's output pin as only it has a valid pin widget
+			if (UAnimStateTransitionNode* TransitionNode = Cast<UAnimStateTransitionNode>(Pin2->GetOwningNode()))
+			{
+				Pin2 = TransitionNode->GetNextState()->GetInputPin();
+			}
+		}
+
+		if (bStartHovered)
+		{
+			SplineOverlapResult = FGraphSplineOverlapResult(Pin2, Pin1, FMath::Min(SquaredDistToPin1, SquaredDistToPin2), SquaredDistToPin2, SquaredDistToPin1, true);
+		}
+		else if (bEndHovered)
+		{
+			SplineOverlapResult = FGraphSplineOverlapResult(Pin1, Pin2, FMath::Min(SquaredDistToPin1, SquaredDistToPin2), SquaredDistToPin1, SquaredDistToPin2, true);
 		}
 
 		// Draw grab handles only in case no relinking operation is performed
-		if (RelinkConnections.IsEmpty() && Params.bUserFlag2)
+		if (RelinkConnections.IsEmpty() && !Params.bUserFlag2)
 		{
-			if (EndHovered)
+			if (bStartHovered)
 			{
-				// Draw solid orange circle behind the arrow head in case the arrow head is hovered (area that enables a relink).
-				FSlateRoundedBoxBrush RoundedBoxBrush = FSlateRoundedBoxBrush(
-					FLinearColor(0.0f, 0.0f, 0.0f, 0.0f), 9.0f, FStyleColors::AccentOrange, 100.0f);
+				// Draw solid selection circle behind the arrow head in case the arrow head is hovered (area that enables a relink).
+				static FSlateRoundedBoxBrush RoundedBoxBrush = FSlateRoundedBoxBrush(FStyleColors::Foreground, ArrowImage->ImageSize.X * 0.5f);
 
 				FSlateDrawElement::MakeBox(DrawElementsList,
 					ArrowLayerID-1, // Draw behind the arrow
-					FPaintGeometry(EndPoint - ArrowRadius, BubbleImage->ImageSize * ZoomFactor, ZoomFactor),
-					&RoundedBoxBrush);
-
-				ArrowHeadColor = FLinearColor::Black;
+					FPaintGeometry(StartPoint - ArrowRadius + (LengthBias * 0.2f), ArrowImage->ImageSize * ZoomFactor, ZoomFactor),
+					&RoundedBoxBrush,
+					ESlateDrawEffect::None,
+			FStyleColors::AccentOrange.GetSpecifiedColor());
 			}
-			else
+			else if (bEndHovered)
 			{
-				// Draw circle around the arrow in case the transition is hovered (mouse close or over transition line or arrow head).
-				const int CircleLineSegments = 16;
-				const float CircleRadius = 10.0f;
-				const FVector2f CircleCenter = EndPoint - UnitDelta * 2.0f;
-				DrawCircle(CircleCenter, CircleRadius, Params.WireColor, CircleLineSegments);
+				// Draw solid selection circle behind the arrow head in case the arrow head is hovered (area that enables a relink).
+				static FSlateRoundedBoxBrush RoundedBoxBrush = FSlateRoundedBoxBrush(FStyleColors::Foreground, ArrowImage->ImageSize.X * 0.5f);
+
+				FSlateDrawElement::MakeBox(DrawElementsList,
+					ArrowLayerID-1, // Draw behind the arrow
+					FPaintGeometry(EndPoint - ArrowRadius - (LengthBias * 0.2f), ArrowImage->ImageSize * ZoomFactor, ZoomFactor),
+					&RoundedBoxBrush,
+					ESlateDrawEffect::None,
+					FStyleColors::AccentOrange.GetSpecifiedColor());
+
+				ArrowHeadColor = FLinearColor::Black;
 			}
 		}
 	}
 
 	// Draw the number of relinked transitions on the preview transition.
-	if (!RelinkConnections.IsEmpty() && !Params.bUserFlag2)
+	if (!RelinkConnections.IsEmpty() && Params.bUserFlag2)
 	{
 		// Get the number of actually relinked transitions.
 		int32 NumRelinkedTransitions = 0;
diff --git a/Engine/Source/Editor/AnimationBlueprintEditor/Private/StateMachineConnectionDrawingPolicy.h b/Engine/Source/Editor/AnimationBlueprintEditor/Private/StateMachineConnectionDrawingPolicy.h
--- a/Engine/Source/Editor/AnimationBlueprintEditor/Private/StateMachineConnectionDrawingPolicy.h
+++ b/Engine/Source/Editor/AnimationBlueprintEditor/Private/StateMachineConnectionDrawingPolicy.h
@@ -58,5 +58,5 @@ protected:
 	void DrawCircle(const FVector2f& Center, float Radius, const FLinearColor& Color, const int NumLineSegments);
 	TArray<FVector2f> TempPoints;
 
-	static constexpr float RelinkHandleHoverRadius = 20.0f;
+	static constexpr float RelinkHandleHoverRadius = 10.0f;
 };
diff --git a/Engine/Source/Editor/DetailCustomizations/Private/AnimStateAliasNodeDetails.cpp b/Engine/Source/Editor/DetailCustomizations/Private/AnimStateAliasNodeDetails.cpp
--- a/Engine/Source/Editor/DetailCustomizations/Private/AnimStateAliasNodeDetails.cpp
+++ b/Engine/Source/Editor/DetailCustomizations/Private/AnimStateAliasNodeDetails.cpp
@@ -77,6 +77,11 @@ void FAnimStateAliasNodeDetails::GetReferenceableStates(const UAnimStateAliasNod
 		auto Node = *NodeIt;
 		OutStates.Add(Node);
 	}
+
+	OutStates.Sort([](TWeakObjectPtr<UAnimStateNodeBase> InA, TWeakObjectPtr<UAnimStateNodeBase> InB)
+	{
+		return InA.Get()->GetStateName() < InB.Get()->GetStateName();
+	});
 }
 
 bool FAnimStateAliasNodeDetails::IsGlobalAlias() const
diff --git a/Engine/Source/Editor/DetailCustomizations/Private/AnimStateNodeDetails.cpp b/Engine/Source/Editor/DetailCustomizations/Private/AnimStateNodeDetails.cpp
--- a/Engine/Source/Editor/DetailCustomizations/Private/AnimStateNodeDetails.cpp
+++ b/Engine/Source/Editor/DetailCustomizations/Private/AnimStateNodeDetails.cpp
@@ -2,6 +2,8 @@
 
 #include "AnimStateNodeDetails.h"
 
+#include "AnimGraphNode_StateResult.h"
+#include "AnimStateNode.h"
 #include "DetailCategoryBuilder.h"
 #include "DetailLayoutBuilder.h"
 #include "DetailWidgetRow.h"
@@ -29,47 +31,48 @@ TSharedRef<IDetailCustomization> FAnimStateNodeDetails::MakeInstance()
 
 void FAnimStateNodeDetails::CustomizeDetails( IDetailLayoutBuilder& DetailBuilder )
 {
-	IDetailCategoryBuilder& SegmentCategory = DetailBuilder.EditCategory("Animation State", LOCTEXT("AnimationStateCategoryTitle", "Animation State") );
+	TArray<TWeakObjectPtr<UAnimStateNode>> StateNodes = DetailBuilder.GetObjectsOfTypeBeingCustomized<UAnimStateNode>();
+	TArray<UObject*> StateResultNodes;
+	for (TWeakObjectPtr<UAnimStateNode> WeakStateNode : StateNodes)
+	{
+		UAnimStateNode* StateNode = WeakStateNode.Get();
+		if (StateNode == nullptr)
+		{
+			continue;
+		}
+
+		UAnimGraphNode_StateResult* StateResultNode = StateNode->GetResultNodeInsideState();
+		if (StateResultNode == nullptr)
+		{
+			continue;
+		}
+
+		StateResultNodes.Add(StateResultNode);
+	}
+
+	IDetailCategoryBuilder& AnimationStateCategory = DetailBuilder.EditCategory("Animation State", LOCTEXT("AnimationState", "Animation State"));
+
+	IDetailPropertyRow* Row = AnimationStateCategory.AddExternalObjects(StateResultNodes, EPropertyLocation::Default, FAddPropertyParams().HideRootObjectNode(true));
+	if (Row != nullptr)
+	{
+		Row->ShouldAutoExpand(true);
+	}
+
+	IDetailCategoryBuilder& DeprecatedCategory = DetailBuilder.EditCategory("Deprecated", LOCTEXT("DeprecatedCategory", "Deprecated"));
+	DeprecatedCategory.InitiallyCollapsed(true);
 	
-	GenerateAnimationStateEventRow(SegmentCategory, LOCTEXT("EnteredAnimationStateEventLabel", "Entered State Event"), TEXT("StateEntered"));
-	GenerateAnimationStateEventRow(SegmentCategory, LOCTEXT("ExitedAnimationStateEventLabel", "Left State Event"), TEXT("StateLeft"));
-	GenerateAnimationStateEventRow(SegmentCategory, LOCTEXT("FullyBlendedAnimationStateEventLabel", "Fully Blended State Event"), TEXT("StateFullyBlended"));
+	GenerateAnimationStateEventRow(DeprecatedCategory, LOCTEXT("EnteredAnimationStateEventLabel", "Entered State Event"), TEXT("StateEntered"));
+	GenerateAnimationStateEventRow(DeprecatedCategory, LOCTEXT("ExitedAnimationStateEventLabel", "Left State Event"), TEXT("StateLeft"));
+	GenerateAnimationStateEventRow(DeprecatedCategory, LOCTEXT("FullyBlendedAnimationStateEventLabel", "Fully Blended State Event"), TEXT("StateFullyBlended"));
 
 	DetailBuilder.HideProperty("StateEntered");
 	DetailBuilder.HideProperty("StateLeft");
 	DetailBuilder.HideProperty("StateFullyBlended");
 }
 
-void FAnimStateNodeDetails::GenerateAnimationStateEventRow(IDetailCategoryBuilder& SegmentCategory,const FText& StateEventLabel, const FString& TransitionName)
+void FAnimStateNodeDetails::GenerateAnimationStateEventRow(IDetailCategoryBuilder& AnimationStateCategory,const FText& StateEventLabel, const FString& TransitionName)
 {
-	const FSlateBrush* WarningIcon = FAppStyle::GetBrush("Icons.WarningWithColor");
-
-	SegmentCategory.AddCustomRow( StateEventLabel)
-		[
-			SNew(SHorizontalBox)
-			.ToolTipText(LOCTEXT("AnimStateEventDeprecatedTooltip", "This is no longer recommended. The recommended approach is to use the anim node function versions."))
-			.Visibility(EVisibility::Visible)
-			
-			+SHorizontalBox::Slot()
-			.VAlign(VAlign_Center)
-			.AutoWidth()
-			[
-				SNew( STextBlock )
-				.Text( StateEventLabel )
-				.Font( IDetailLayoutBuilder::GetDetailFontBold() )
-			]
-
-			+SHorizontalBox::Slot()
-			.VAlign(VAlign_Center)
-			.AutoWidth()
-			.HAlign(HAlign_Right)
-			.Padding(4.0f, 0.0f, 0.0f, 0.0f)
-			[	SNew(SImage)
-				.Image(WarningIcon)
-			]
-		];
-		
-	CreateTransitionEventPropertyWidgets(SegmentCategory, TransitionName);
+	CreateTransitionEventPropertyWidgets(AnimationStateCategory, TransitionName, true);
 }
 
 #undef LOCTEXT_NAMESPACE
diff --git a/Engine/Source/Editor/DetailCustomizations/Private/AnimStateNodeDetails.h b/Engine/Source/Editor/DetailCustomizations/Private/AnimStateNodeDetails.h
--- a/Engine/Source/Editor/DetailCustomizations/Private/AnimStateNodeDetails.h
+++ b/Engine/Source/Editor/DetailCustomizations/Private/AnimStateNodeDetails.h
@@ -19,6 +19,6 @@ public:
 
 private:
 	
-	void GenerateAnimationStateEventRow(IDetailCategoryBuilder& SegmentCategory, const FText & StateEventLabel, const FString & TransitionName);
+	void GenerateAnimationStateEventRow(IDetailCategoryBuilder& AnimationStateCategory, const FText & StateEventLabel, const FString & TransitionName);
 };
 
diff --git a/Engine/Source/Editor/DetailCustomizations/Private/AnimTransitionNodeDetails.cpp b/Engine/Source/Editor/DetailCustomizations/Private/AnimTransitionNodeDetails.cpp
--- a/Engine/Source/Editor/DetailCustomizations/Private/AnimTransitionNodeDetails.cpp
+++ b/Engine/Source/Editor/DetailCustomizations/Private/AnimTransitionNodeDetails.cpp
@@ -53,6 +53,7 @@
 #include "Widgets/Input/STextEntryPopup.h"
 #include "Widgets/SBoxPanel.h"
 #include "Widgets/SWindow.h"
+#include "Widgets/Images/SImage.h"
 #include "Widgets/Text/STextBlock.h"
 
 class SWidget;
@@ -276,34 +277,14 @@ void FAnimTransitionNodeDetails::CustomizeDetails( IDetailLayoutBuilder& DetailB
 
 		IDetailCategoryBuilder& NotificationCategory = DetailBuilder.EditCategory("Notifications", LOCTEXT("NotificationsCategoryTitle", "Notifications") );
 
-		NotificationCategory.AddCustomRow( LOCTEXT("StartTransitionEventPropertiesCategoryLabel", "Start Transition Event") )
-		[
-			SNew( STextBlock )
-			.Text( LOCTEXT("StartTransitionEventPropertiesCategoryLabel", "Start Transition Event") )
-			.Font( IDetailLayoutBuilder::GetDetailFontBold() )
-		];
-		CreateTransitionEventPropertyWidgets(NotificationCategory, TEXT("TransitionStart"));
-
-
-		NotificationCategory.AddCustomRow( LOCTEXT("EndTransitionEventPropertiesCategoryLabel", "End Transition Event" ) ) 
-		[
-			SNew( STextBlock )
-			.Text( LOCTEXT("EndTransitionEventPropertiesCategoryLabel", "End Transition Event" ) )
-			.Font( IDetailLayoutBuilder::GetDetailFontBold() )
-		];
-		CreateTransitionEventPropertyWidgets(NotificationCategory, TEXT("TransitionEnd"));
-
-		NotificationCategory.AddCustomRow( LOCTEXT("InterruptTransitionEventPropertiesCategoryLabel", "Interrupt Transition Event") )
-		[
-			SNew( STextBlock )
-			.Text( LOCTEXT("InterruptTransitionEventPropertiesCategoryLabel", "Interrupt Transition Event") )
-			.Font( IDetailLayoutBuilder::GetDetailFontBold() )
-		];
-		CreateTransitionEventPropertyWidgets(NotificationCategory, TEXT("TransitionInterrupt"));
+		CreateTransitionEventPropertyWidgets(NotificationCategory, TEXT("TransitionStart"), false);
+		CreateTransitionEventPropertyWidgets(NotificationCategory, TEXT("TransitionEnd"), false);
+		CreateTransitionEventPropertyWidgets(NotificationCategory, TEXT("TransitionInterrupt"), false);
 	}
 
 	DetailBuilder.HideProperty(GET_MEMBER_NAME_CHECKED(UAnimStateTransitionNode, TransitionStart));
 	DetailBuilder.HideProperty(GET_MEMBER_NAME_CHECKED(UAnimStateTransitionNode, TransitionEnd));
+	DetailBuilder.HideProperty(GET_MEMBER_NAME_CHECKED(UAnimStateTransitionNode, TransitionInterrupt));
 }
 END_SLATE_FUNCTION_BUILD_OPTIMIZATION
 
@@ -517,12 +498,57 @@ EVisibility FAnimTransitionNodeDetails::GetBlendGraphButtonVisibility(bool bMult
 }
 
 
-void FAnimTransitionNodeDetails::CreateTransitionEventPropertyWidgets(IDetailCategoryBuilder& TransitionCategory, FString TransitionName)
+void FAnimTransitionNodeDetails::CreateTransitionEventPropertyWidgets(IDetailCategoryBuilder& TransitionCategory, FString TransitionName, bool bInWarnAboutDeprecation)
 {
 	TSharedPtr<IPropertyHandle> NameProperty = TransitionCategory.GetParentLayout().GetProperty(*(TransitionName + TEXT(".NotifyName")));
 
-	TransitionCategory.AddProperty( NameProperty )
-		.DisplayName( LOCTEXT("CreateTransition_CustomBlueprintEvent", "Custom Blueprint Event") );
+	IDetailPropertyRow& Row = TransitionCategory.AddProperty(NameProperty);
+
+	if (bInWarnAboutDeprecation)
+	{
+		const FSlateBrush* WarningIcon = FAppStyle::GetBrush("Icons.WarningWithColor");
+
+		Row.CustomWidget()
+		.NameContent()
+		[
+			SNew(SHorizontalBox)
+			.ToolTipText(LOCTEXT("AnimStateEventDeprecatedTooltip", "This is no longer recommended. The recommended approach is to use the anim node function versions."))
+			.Visibility(EVisibility::Visible)
+			
+			+SHorizontalBox::Slot()
+			.VAlign(VAlign_Center)
+			.AutoWidth()
+			[
+				NameProperty->GetParentHandle()->CreatePropertyNameWidget()
+			]
+
+			+SHorizontalBox::Slot()
+			.VAlign(VAlign_Center)
+			.AutoWidth()
+			.HAlign(HAlign_Right)
+			.Padding(4.0f, 0.0f, 0.0f, 0.0f)
+			[
+				SNew(SImage)
+				.Image(WarningIcon)
+			]
+		]
+		.ValueContent()
+		[
+			NameProperty->CreatePropertyValueWidget()
+		];
+	}
+	else
+	{
+		Row.CustomWidget()
+		.NameContent()
+		[
+			NameProperty->GetParentHandle()->CreatePropertyNameWidget()
+		]
+		.ValueContent()
+		[
+			NameProperty->CreatePropertyValueWidget()
+		];
+	}
 }
 
 TSharedRef<SWidget> FAnimTransitionNodeDetails::GetWidgetForInlineShareMenu(const TAttribute<FText>& InSharedNameText, const TAttribute<bool>& bInIsCurrentlyShared, FOnClicked PromoteClick, FOnClicked DemoteClick, FOnGetContent GetContentMenu)
diff --git a/Engine/Source/Editor/DetailCustomizations/Private/AnimTransitionNodeDetails.h b/Engine/Source/Editor/DetailCustomizations/Private/AnimTransitionNodeDetails.h
--- a/Engine/Source/Editor/DetailCustomizations/Private/AnimTransitionNodeDetails.h
+++ b/Engine/Source/Editor/DetailCustomizations/Private/AnimTransitionNodeDetails.h
@@ -33,7 +33,7 @@ public:
 	virtual void CustomizeDetails( IDetailLayoutBuilder& DetailBuilder ) override;
 
 protected:
-	void CreateTransitionEventPropertyWidgets(IDetailCategoryBuilder& TransitionCategory, FString TransitionName);
+	void CreateTransitionEventPropertyWidgets(IDetailCategoryBuilder& TransitionCategory, FString TransitionName, bool bInWarnAboutDeprecation);
 
 	FReply OnClickEditBlendGraph();
 	EVisibility GetBlendGraphButtonVisibility(bool bMultiSelect) const;
diff --git a/Engine/Source/Editor/EditorStyle/Private/StarshipStyle.cpp b/Engine/Source/Editor/EditorStyle/Private/StarshipStyle.cpp
--- a/Engine/Source/Editor/EditorStyle/Private/StarshipStyle.cpp
+++ b/Engine/Source/Editor/EditorStyle/Private/StarshipStyle.cpp
@@ -3980,9 +3980,29 @@ void FStarshipEditorStyle::FStyle::SetupGraphEditorStyles()
 				Set( "Graph.StateNode.Icon", new IMAGE_BRUSH_SVG( "Starship/AnimationBlueprintEditor/AnimationState", Icon16x16 ) );
 				Set( "Graph.ConduitNode.Icon", new IMAGE_BRUSH_SVG( "Starship/AnimationBlueprintEditor/AnimationConduit", Icon16x16 ) );
 				Set( "Graph.AliasNode.Icon", new IMAGE_BRUSH_SVG( "Starship/AnimationBlueprintEditor/AnimationAlias", Icon16x16 ) );
+				Set( "Graph.EntryNode.Icon", new IMAGE_BRUSH_SVG( "Starship/Common/Event", Icon16x16 ) );
 
 				Set( "Graph.StateNode.Pin.BackgroundHovered", new BOX_BRUSH( "/Persona/StateMachineEditor/StateNode_Pin_HoverCue", FMargin(12.0f/64.0f,12.0f/64.0f,12.0f/64.0f,12.0f/64.0f)));
 				Set( "Graph.StateNode.Pin.Background", new FSlateNoResource() );
+
+				// Anim state node specific brushes
+				Set( "Graph.AnimStateNode.Body", new FSlateRoundedBoxBrush(FLinearColor(0.8f, 0.8f, 0.8f, 1.0f), 7.0f));
+				Set( "Graph.AnimStateNode.ColorSpill", new FSlateRoundedBoxBrush(FLinearColor(0.2f, 0.2f, 0.2f, 1.0f), 7.0f));
+				Set( "Graph.AnimStateNode.Selection", new FSlateRoundedBoxBrush(FLinearColor(0.0f, 0.0f, 0.0f, 0.0f), 7.0f, FStyleColors::AccentOrange, 3.0f));
+				Set( "Graph.AnimAliasNode.Body", new FSlateRoundedBoxBrush(FLinearColor(0.8f, 0.8f, 0.8f, 1.0f)));
+				Set( "Graph.AnimAliasNode.ColorSpill", new FSlateRoundedBoxBrush(FLinearColor(0.2f, 0.2f, 0.2f, 1.0f)));
+				Set( "Graph.AnimAliasNode.Selection", new FSlateRoundedBoxBrush(FLinearColor(0.0f, 0.0f, 0.0f, 0.0f), FStyleColors::AccentOrange, 3.0f));
+				Set( "Graph.AnimAliasNode.Shadow", new BOX_BRUSH( "/Persona/StateMachineEditor/AliasNode_shadow", FMargin(26.0f/64.0f) ));
+				Set( "Graph.AnimConduitNode.Body", new FSlateRoundedBoxBrush(FLinearColor(0.8f, 0.8f, 0.8f, 1.0f), 2.0f));
+				Set( "Graph.AnimConduitNode.ColorSpill", new FSlateRoundedBoxBrush(FLinearColor(0.2f, 0.2f, 0.2f, 1.0f), 2.0f));
+				Set( "Graph.AnimConduitNode.Selection", new FSlateRoundedBoxBrush(FLinearColor(0.0f, 0.0f, 0.0f, 0.0f), 2.0f, FStyleColors::AccentOrange, 3.0f));
+				
+				Set( "Graph.AnimStateNode.Pin.BackgroundHovered", new FSlateRoundedBoxBrush(FLinearColor(0.8f, 0.8f, 0.8f, 0.5f), 7.0f));
+				Set( "Graph.AnimAliasNode.Pin.BackgroundHovered", new FSlateRoundedBoxBrush(FLinearColor(0.8f, 0.8f, 0.8f, 0.5f)));
+				Set( "Graph.AnimConduitNode.Pin.BackgroundHovered", new FSlateRoundedBoxBrush(FLinearColor(0.8f, 0.8f, 0.8f, 0.5f), 2.0f));
+				Set( "Graph.AnimStateNode.Pin.Background", new FSlateNoResource() );
+
+				Set( "Graph.AnimStateNode.ConnectionArrow", new IMAGE_BRUSH_SVG( "/Persona/StateMachineEditor/ArrowHead", Icon16x16 ) );
 			}
 
 			{
@@ -4009,13 +4029,41 @@ void FStarshipEditorStyle::FStyle::SetupGraphEditorStyles()
 					.SetEditableTextBoxStyle(GraphStateNodeTitleEditableText)
 					);
 			}
+			
+			{	
+				FTextBlockStyle AnimGraphStateNodeTitle = FTextBlockStyle(NormalText)
+					.SetFont( DEFAULT_FONT( "Bold", 12 ) )
+					.SetColorAndOpacity( FLinearColor(230.0f/255.0f,230.0f/255.0f,230.0f/255.0f) );
+				Set( "Graph.AnimStateNode.NodeTitle", AnimGraphStateNodeTitle );
+
+				FEditableTextBoxStyle AnimGraphStateNodeTitleEditableText = FEditableTextBoxStyle()
+					.SetTextStyle(AnimGraphStateNodeTitle)
+					.SetBackgroundImageNormal(FSlateRoundedBoxBrush(FStyleColors::Input, CoreStyleConstants::InputFocusRadius, FStyleColors::InputOutline, CoreStyleConstants::InputFocusThickness))
+					.SetBackgroundImageHovered(FSlateRoundedBoxBrush(FStyleColors::Input, CoreStyleConstants::InputFocusRadius, FStyleColors::Hover, CoreStyleConstants::InputFocusThickness))
+					.SetBackgroundImageFocused(FSlateRoundedBoxBrush(FStyleColors::Input, CoreStyleConstants::InputFocusRadius, FStyleColors::Primary, CoreStyleConstants::InputFocusThickness))
+					.SetBackgroundImageReadOnly(FSlateRoundedBoxBrush(FStyleColors::Header, CoreStyleConstants::InputFocusRadius, FStyleColors::InputOutline, CoreStyleConstants::InputFocusThickness))
+					.SetFocusedForegroundColor(FStyleColors::White)
+					.SetScrollBarStyle( ScrollBar );
+				Set( "Graph.AnimStateNode.NodeTitleEditableText", AnimGraphStateNodeTitle );
+
+				Set( "Graph.AnimStateNode.NodeTitleInlineEditableText", FInlineEditableTextBlockStyle()
+					.SetTextStyle(AnimGraphStateNodeTitle)
+					.SetEditableTextBoxStyle(AnimGraphStateNodeTitleEditableText)
+					);
+			}
 
-			// Transition node
+			// Transition nodes
 			{
 				FMargin TestMargin(16.f/64.f, 16.f/28.f, 16.f/64.f, 4.f/28.f);
 				Set( "Graph.TransitionNode.ColorSpill", new BOX_BRUSH( "/Persona/StateMachineEditor/Trans_Node_ColorSpill", TestMargin ) );
 				Set( "Graph.TransitionNode.Icon", new IMAGE_BRUSH( "/Persona/StateMachineEditor/Trans_Node_Icon", FVector2f(25.f,25.f) ) );
 				Set( "Graph.TransitionNode.Icon_Inertialization", new IMAGE_BRUSH( "/Persona/StateMachineEditor/Trans_Node_Icon_Inertialization", FVector2f(25.f,25.f) ) );
+
+				Set( "Graph.AnimTransitionNode.ColorSpill", new FSlateRoundedBoxBrush(FLinearColor(0.8f, 0.8f, 0.8f, 1.0f)));
+				Set( "Graph.AnimTransitionNode.Selection", new FSlateRoundedBoxBrush(FLinearColor(0.0f, 0.0f, 0.0f, 0.0f), FStyleColors::AccentOrange, 3.0f));
+				Set( "Graph.AnimTransitionNode.Icon", new CORE_IMAGE_BRUSH_SVG("Starship/Common/arrow-right", FVector2f(25.f, 25.f)));
+				Set( "Graph.AnimTransitionNode.Icon_Inertialization", new CORE_IMAGE_BRUSH_SVG("Starship/Common/arrow-right", FVector2f(25.f, 25.f)));
+				Set( "Graph.AnimTransitionNode.Shadow", new BOX_BRUSH( "/Persona/StateMachineEditor/TransitionNode_shadow", FMargin(26.0f/64.0f) ) );
 			}
 
 			// Transition rule tooltip name 
diff --git a/Engine/Source/Editor/Kismet/Private/BlueprintDetailsCustomization.cpp b/Engine/Source/Editor/Kismet/Private/BlueprintDetailsCustomization.cpp
--- a/Engine/Source/Editor/Kismet/Private/BlueprintDetailsCustomization.cpp
+++ b/Engine/Source/Editor/Kismet/Private/BlueprintDetailsCustomization.cpp
@@ -7903,7 +7903,8 @@ void FBlueprintComponentDetails::OnSocketSelection(FName SocketName)
 BEGIN_SLATE_FUNCTION_BUILD_OPTIMIZATION
 void FBlueprintGraphNodeDetails::CustomizeDetails( IDetailLayoutBuilder& DetailLayout )
 {
-	const TArray<TWeakObjectPtr<UObject>>& SelectedObjects = DetailLayout.GetSelectedObjects();
+	TArray<TWeakObjectPtr<UObject>> SelectedObjects;
+	DetailLayout.GetObjectsBeingCustomized(SelectedObjects);
 	if( SelectedObjects.Num() == 1 )
 	{
 		if (SelectedObjects[0].IsValid() && SelectedObjects[0]->IsA<UEdGraphNode>())
diff --git a/Engine/Source/Runtime/Engine/Classes/Animation/AnimNode_StateMachine.h b/Engine/Source/Runtime/Engine/Classes/Animation/AnimNode_StateMachine.h
--- a/Engine/Source/Runtime/Engine/Classes/Animation/AnimNode_StateMachine.h
+++ b/Engine/Source/Runtime/Engine/Classes/Animation/AnimNode_StateMachine.h
@@ -222,6 +222,13 @@ private:
 
 	TArray<FGraphTraversalCounter> StateCacheBoneCounters;
 
+	// The last time that each respective state was entered
+	TArray<double> LastStateEntryTime; 
+
+#if WITH_EDITOR
+	int32 DebugForcedStateIndex = INDEX_NONE;
+#endif
+
 public:
 	FAnimNode_StateMachine()
 		: StateMachineIndexInClass(0)
@@ -258,7 +265,7 @@ public:
 	/** Cache the internal machine description */
 	ENGINE_API void CacheMachineDescription(IAnimClassInterface* AnimBlueprintClass);
 
-	ENGINE_API void SetState(const FAnimationBaseContext& Context, int32 NewStateIndex);
+	ENGINE_API void SetState(const FAnimationBaseContext& Context, int32 NewStateIndex, bool bAllowReEntry);
 	ENGINE_API void TransitionToState(const FAnimationUpdateContext& Context, const FAnimationTransitionBetweenStates& TransitionInfo, const FAnimationPotentialTransition* BakedTransitionInfo = nullptr);
 	ENGINE_API const int32 GetStateIndex(FName StateName) const;
 
@@ -345,7 +352,15 @@ protected:
 	/** Removes all marked events that are queued */
 	ENGINE_API void ConsumeMarkedTransitionEvents();
 
+#if WITH_EDITOR
+	void DebugForceState(int32 InStateIndex)
+	{
+		DebugForcedStateIndex = InStateIndex;
+	}
+#endif
+
 public:
 	friend struct FAnimInstanceProxy;
 	friend class UAnimationStateMachineLibrary;
+	friend class FAnimationBlueprintEditor;
 };
diff --git a/Engine/Source/Runtime/Engine/Classes/Animation/AnimStateMachineTypes.h b/Engine/Source/Runtime/Engine/Classes/Animation/AnimStateMachineTypes.h
--- a/Engine/Source/Runtime/Engine/Classes/Animation/AnimStateMachineTypes.h
+++ b/Engine/Source/Runtime/Engine/Classes/Animation/AnimStateMachineTypes.h
@@ -45,9 +45,9 @@ namespace ETransitionLogicType
 {
 	enum Type : int
 	{
-		TLT_StandardBlend UMETA(DisplayName="Standard Blend"),
-		TLT_Inertialization UMETA(DisplayName = "Inertialization"),
-		TLT_Custom UMETA(DisplayName="Custom")
+		TLT_StandardBlend UMETA(DisplayName="Standard Blend", ToolTip="Blend smoothly from source state to destination state.\nBoth states update during the transition.\nFalls back to Inertialization on re-entry to an already active state when Fall Back to Inertialization is true."),
+		TLT_Inertialization UMETA(DisplayName = "Inertialization", ToolTip="Use inertialization to extrapolate when blending between states.\nOnly one state is active at a time.\nRequires an Inertialization or Dead Blending node rootwards of this node in the graph."),
+		TLT_Custom UMETA(DisplayName="Custom", ToolTip="Use a custom graph to define exactly how the blend works.")
 	};
 }
 
@@ -187,6 +187,12 @@ struct FAnimationTransitionBetweenStates : public FAnimationStateBase
 {
 	GENERATED_USTRUCT_BODY()
 
+	UPROPERTY()
+	TObjectPtr<UCurveFloat> CustomCurve;
+
+	UPROPERTY()
+	TObjectPtr<UBlendProfile> BlendProfile;
+	
 	// Transition-only: State being transitioned from
 	UPROPERTY()
 	int32 PreviousState;
@@ -198,6 +204,9 @@ struct FAnimationTransitionBetweenStates : public FAnimationStateBase
 	UPROPERTY()
 	float CrossfadeDuration;
 
+	UPROPERTY()
+	float MinTimeBeforeReentry;
+
 	UPROPERTY()
 	int32 StartNotify;
 
@@ -211,31 +220,30 @@ struct FAnimationTransitionBetweenStates : public FAnimationStateBase
 	EAlphaBlendOption BlendMode;
 
 	UPROPERTY()
-	TObjectPtr<UCurveFloat> CustomCurve;
-
-	UPROPERTY()
-	TObjectPtr<UBlendProfile> BlendProfile;
+	TEnumAsByte<ETransitionLogicType::Type> LogicType;
 
 	UPROPERTY()
-	TEnumAsByte<ETransitionLogicType::Type> LogicType;
+	uint8 bAllowInertializationForSelfTransitions : 1;
 
 #if WITH_EDITORONLY_DATA
 	// This is only needed for the baking process, to denote which baked transitions need to reverse their prev/next state in the final step
-	bool ReverseTransition;
+	uint8 ReverseTransition : 1;
 #endif
 
 	FAnimationTransitionBetweenStates()
 		: FAnimationStateBase()
+		, CustomCurve(nullptr)
+		, BlendProfile(nullptr)
 		, PreviousState(INDEX_NONE)
 		, NextState(INDEX_NONE)
-		, CrossfadeDuration(0.0f)
+		, CrossfadeDuration(-1.0f)
+		, MinTimeBeforeReentry(-1.0f)
 		, StartNotify(INDEX_NONE)
 		, EndNotify(INDEX_NONE)
 		, InterruptNotify(INDEX_NONE)
 		, BlendMode(EAlphaBlendOption::CubicInOut)
-		, CustomCurve(nullptr)
-		, BlendProfile(nullptr)
 		, LogicType(ETransitionLogicType::TLT_StandardBlend)
+		, bAllowInertializationForSelfTransitions(false)
 #if WITH_EDITOR
 		, ReverseTransition(false)
 #endif
@@ -300,14 +308,22 @@ struct FBakedAnimationState
 {
 	GENERATED_USTRUCT_BODY()
 
-	// The name of this state
+	// Indices into the property array for player nodes in the state
 	UPROPERTY()
-	FName StateName;
+	TArray<int32> PlayerNodeIndices;
 
+	// Indices into the property array for layer nodes in the state
+	UPROPERTY()
+	TArray<int32> LayerNodeIndices;
+	
 	// Set of legal transitions out of this state; already in priority order
 	UPROPERTY()
 	TArray<FBakedStateExitTransition> Transitions;
 
+	// The name of this state
+	UPROPERTY()
+	FName StateName;
+
 	// The root node index (into the AnimNodeProperties array of the UAnimBlueprintGeneratedClass)
 	UPROPERTY()
 	int32 StateRootNodeIndex;
@@ -320,34 +336,26 @@ struct FBakedAnimationState
 
 	UPROPERTY()
 	int32 FullyBlendedNotify;
-
-	UPROPERTY()
-	bool bIsAConduit;
-
+	
 	UPROPERTY()
 	int32 EntryRuleNodeIndex;
 
-	// Indices into the property array for player nodes in the state
-	UPROPERTY()
-	TArray<int32> PlayerNodeIndices;
-
-	// Indices into the property array for layer nodes in the state
-	UPROPERTY()
-	TArray<int32> LayerNodeIndices;
-
 	// Whether or not this state will ALWAYS reset it's state on reentry, regardless of remaining weight
 	UPROPERTY()
 	bool bAlwaysResetOnEntry;
 
+	UPROPERTY()
+	bool bIsAConduit;
+
 public:
 	FBakedAnimationState()
 		: StateRootNodeIndex(INDEX_NONE)
 		, StartNotify(INDEX_NONE)
 		, EndNotify(INDEX_NONE)
 		, FullyBlendedNotify(INDEX_NONE)
-		, bIsAConduit(false)
 		, EntryRuleNodeIndex(INDEX_NONE)
 		, bAlwaysResetOnEntry(false)
+		, bIsAConduit(false)
 	{}
 };
 
diff --git a/Engine/Source/Runtime/Engine/Private/Animation/AnimInstanceProxy.cpp b/Engine/Source/Runtime/Engine/Private/Animation/AnimInstanceProxy.cpp
--- a/Engine/Source/Runtime/Engine/Private/Animation/AnimInstanceProxy.cpp
+++ b/Engine/Source/Runtime/Engine/Private/Animation/AnimInstanceProxy.cpp
@@ -2814,9 +2814,9 @@ bool FAnimInstanceProxy::RequestTransitionEvent(const FName& EventName, const do
 	return true;
 }
 
-bool FAnimInstanceProxy::QueryTransitionEvent(int32 MachineIndex, int32 TransitionIndex, const FName& EventName) const
+bool FAnimInstanceProxy::QueryTransitionEvent(int32 NodeIndex, int32 TransitionIndex, const FName& EventName) const
 {
-	const FAnimNode_StateMachine* StateMachine = GetStateMachineInstance(MachineIndex);
+	const FAnimNode_StateMachine* StateMachine = GetStateMachineInstance(NodeIndex);
 	if (StateMachine)
 	{
 		return StateMachine->QueryTransitionEvent(TransitionIndex, EventName);
@@ -2840,13 +2840,13 @@ void FAnimInstanceProxy::ClearAllTransitionEvents()
 	});
 }
 
-bool FAnimInstanceProxy::QueryAndMarkTransitionEvent(int32 MachineIndex, int32 TransitionIndex, const FName& EventName)
+bool FAnimInstanceProxy::QueryAndMarkTransitionEvent(int32 NodeIndex, int32 TransitionIndex, const FName& EventName)
 {
 	if (AnimClassInterface)
 	{
 		const TArray<FStructProperty*>& AnimNodeProperties = AnimClassInterface->GetAnimNodeProperties();
 
-		FStructProperty* Property = AnimNodeProperties[AnimNodeProperties.Num() - 1 - MachineIndex];
+		FStructProperty* Property = AnimNodeProperties[AnimNodeProperties.Num() - 1 - NodeIndex];
 		if (Property && Property->Struct->IsChildOf(FAnimNode_StateMachine::StaticStruct()))
 		{
 			FAnimNode_StateMachine* StateMachine = Property->ContainerPtrToValuePtr<FAnimNode_StateMachine>(AnimInstanceObject);
@@ -2860,9 +2860,9 @@ bool FAnimInstanceProxy::QueryAndMarkTransitionEvent(int32 MachineIndex, int32 T
 	return false;
 }
 
-const FAnimNode_AssetPlayerRelevancyBase* FAnimInstanceProxy::GetRelevantAssetPlayerInterfaceFromState(int32 MachineIndex, int32 StateIndex) const
+const FAnimNode_AssetPlayerRelevancyBase* FAnimInstanceProxy::GetRelevantAssetPlayerInterfaceFromState(int32 NodeIndex, int32 StateIndex) const
 {
-	if (const FAnimNode_StateMachine* StateMachine = GetStateMachineInstance(MachineIndex))
+	if (const FAnimNode_StateMachine* StateMachine = GetStateMachineInstance(NodeIndex))
 	{
 		return StateMachine->GetRelevantAssetPlayerInterfaceFromState(this, StateMachine->GetStateInfo(StateIndex));
 	}
@@ -2870,19 +2870,20 @@ const FAnimNode_AssetPlayerRelevancyBase* FAnimInstanceProxy::GetRelevantAssetPl
 	return nullptr;
 }
 
-const FAnimNode_StateMachine* FAnimInstanceProxy::GetStateMachineInstance(int32 MachineIndex) const
+const FAnimNode_StateMachine* FAnimInstanceProxy::GetStateMachineInstance(int32 NodeIndex) const
 {
 	if (AnimClassInterface)
 	{
 		const TArray<FStructProperty*>& AnimNodeProperties = AnimClassInterface->GetAnimNodeProperties();
-		if ((MachineIndex >= 0) && (MachineIndex < AnimNodeProperties.Num()))
+		if ((NodeIndex >= 0) && (NodeIndex < AnimNodeProperties.Num()))
 		{
-			const int32 InstancePropertyIndex = AnimNodeProperties.Num() - 1 - MachineIndex;
+			const int32 InstancePropertyIndex = AnimNodeProperties.Num() - 1 - NodeIndex;
 
 			FStructProperty* MachineInstanceProperty = AnimNodeProperties[InstancePropertyIndex];
-			checkSlow(MachineInstanceProperty->Struct->IsChildOf(FAnimNode_StateMachine::StaticStruct()));
-
-			return MachineInstanceProperty->ContainerPtrToValuePtr<FAnimNode_StateMachine>(AnimInstanceObject);
+			if(MachineInstanceProperty->Struct->IsChildOf(FAnimNode_StateMachine::StaticStruct()))
+			{
+				return MachineInstanceProperty->ContainerPtrToValuePtr<FAnimNode_StateMachine>(AnimInstanceObject);
+			}
 		}
 	}
 
@@ -3142,20 +3143,20 @@ int32 FAnimInstanceProxy::GetStateMachineIndex(FName MachineName) const
 	return INDEX_NONE;
 }
 
-int32  FAnimInstanceProxy::GetStateMachineIndex(FAnimNode_StateMachine* StateMachine) const
+int32 FAnimInstanceProxy::GetStateMachineIndex(FAnimNode_StateMachine* StateMachine) const
 {
 	if (AnimClassInterface)
 	{
 		const TArray<FStructProperty*>& AnimNodeProperties = AnimClassInterface->GetAnimNodeProperties();
-		for (int32 MachineIndex = 0; MachineIndex < AnimNodeProperties.Num(); MachineIndex++)
+		for (int32 NodeIndex = 0; NodeIndex < AnimNodeProperties.Num(); NodeIndex++)
 		{
-			FStructProperty* Property = AnimNodeProperties[AnimNodeProperties.Num() - 1 - MachineIndex];
+			FStructProperty* Property = AnimNodeProperties[AnimNodeProperties.Num() - 1 - NodeIndex];
 			if (Property && Property->Struct->IsChildOf(FAnimNode_StateMachine::StaticStruct()))
 			{
 				FAnimNode_StateMachine* CurStateMachine = Property->ContainerPtrToValuePtr<FAnimNode_StateMachine>(AnimInstanceObject);
 				if (CurStateMachine == StateMachine)
 				{
-					return MachineIndex;
+					return NodeIndex;
 				}
 			}
 		}
@@ -3163,14 +3164,14 @@ int32  FAnimInstanceProxy::GetStateMachineIndex(FAnimNode_StateMachine* StateMac
 	return INDEX_NONE; 	
 }
 
-void FAnimInstanceProxy::GetStateMachineIndexAndDescription(FName InMachineName, int32& OutMachineIndex, const FBakedAnimationStateMachine** OutMachineDescription) const
+void FAnimInstanceProxy::GetStateMachineIndexAndDescription(FName InMachineName, int32& OutNodeIndex, const FBakedAnimationStateMachine** OutMachineDescription) const
 {
 	if (AnimClassInterface)
 	{
 		const TArray<FStructProperty*>& AnimNodeProperties = AnimClassInterface->GetAnimNodeProperties();
-		for (int32 MachineIndex = 0; MachineIndex < AnimNodeProperties.Num(); MachineIndex++)
+		for (int32 NodeIndex = 0; NodeIndex < AnimNodeProperties.Num(); NodeIndex++)
 		{
-			FStructProperty* Property = AnimNodeProperties[AnimNodeProperties.Num() - 1 - MachineIndex];
+			FStructProperty* Property = AnimNodeProperties[AnimNodeProperties.Num() - 1 - NodeIndex];
 			if (Property && Property->Struct->IsChildOf(FAnimNode_StateMachine::StaticStruct()))
 			{
 				const FAnimNode_StateMachine* StateMachine = Property->ContainerPtrToValuePtr<FAnimNode_StateMachine>(AnimInstanceObject);
@@ -3180,7 +3181,7 @@ void FAnimInstanceProxy::GetStateMachineIndexAndDescription(FName InMachineName,
 					{
 						if (MachineDescription->MachineName == InMachineName)
 						{
-							OutMachineIndex = MachineIndex;
+							OutNodeIndex = NodeIndex;
 							if (OutMachineDescription)
 							{
 								*OutMachineDescription = MachineDescription;
@@ -3193,7 +3194,7 @@ void FAnimInstanceProxy::GetStateMachineIndexAndDescription(FName InMachineName,
 		}
 	}
 
-	OutMachineIndex = INDEX_NONE;
+	OutNodeIndex = INDEX_NONE;
 	if (OutMachineDescription)
 	{
 		*OutMachineDescription = nullptr;
diff --git a/Engine/Source/Runtime/Engine/Private/Animation/AnimNode_StateMachine.cpp b/Engine/Source/Runtime/Engine/Private/Animation/AnimNode_StateMachine.cpp
--- a/Engine/Source/Runtime/Engine/Private/Animation/AnimNode_StateMachine.cpp
+++ b/Engine/Source/Runtime/Engine/Private/Animation/AnimNode_StateMachine.cpp
@@ -19,6 +19,7 @@
 #include "Animation/SkeletonRemapping.h"
 #include "Animation/SkeletonRemappingRegistry.h"
 #include "HAL/IConsoleManager.h"
+#include "Engine/World.h"
 
 #include UE_INLINE_GENERATED_CPP_BY_NAME(AnimNode_StateMachine)
 
@@ -271,9 +272,12 @@ void FAnimNode_StateMachine::Initialize_AnyThread(const FAnimationInitializeCont
 
 			StateCacheBoneCounters.Reset(Machine->States.Num());
 			StateCacheBoneCounters.AddDefaulted(Machine->States.Num());
-		
+
+			LastStateEntryTime.Reset(Machine->States.Num());
+			LastStateEntryTime.AddZeroed(Machine->States.Num());
+
 			// Move to the default state
-			SetState(Context, Machine->InitialState);
+			SetState(Context, Machine->InitialState, false);
 
 			// initialize first update
 			bFirstUpdate = true;
@@ -434,6 +438,17 @@ void FAnimNode_StateMachine::Update_AnyThread(const FAnimationUpdateContext& Con
 	FScopeCycleCounter MachineNameCycleCounter(Machine->GetStatID());
 #endif // STATS
 
+#if WITH_EDITOR
+	if (DebugForcedStateIndex != INDEX_NONE)
+	{
+		if (PRIVATE_MachineDescription->States.IsValidIndex(DebugForcedStateIndex))
+		{
+			SetState(Context, DebugForcedStateIndex, true);
+		}
+		DebugForcedStateIndex = INDEX_NONE;
+	}
+#endif
+
 	// On the first update, call the initial state's Start Notify.
 	if (bFirstUpdate)
 	{
@@ -443,7 +458,6 @@ void FAnimNode_StateMachine::Update_AnyThread(const FAnimationUpdateContext& Con
 
 	bool bFoundValidTransition = false;
 	int32 TransitionCountThisFrame = 0;
-	int32 TransitionIndex = INDEX_NONE;
 
 	// Look for legal transitions to take; can move across multiple states in one frame (up to MaxTransitionsPerFrame)
 	do
@@ -459,10 +473,27 @@ void FAnimNode_StateMachine::Update_AnyThread(const FAnimationUpdateContext& Con
 			TArray<int32, TInlineAllocator<4>> VisitedStateIndices;
 			FindValidTransition(Context, GetStateInfo(), /*Out*/ PotentialTransition, /*Out*/ VisitedStateIndices);
 		}
-				
-		// If transition is valid and not waiting on other conditions
-		// and we're not doing a transition to self
-		if (PotentialTransition.IsValid() && PotentialTransition.TargetState != CurrentState)
+
+		auto CanTransition = [this, &PotentialTransition]()
+		{
+			if (!PotentialTransition.IsValid())
+			{
+				return false;
+			}
+
+			const FAnimationTransitionBetweenStates& ReferenceTransition = GetTransitionInfo(PotentialTransition.TransitionRule->TransitionIndex); //-V595
+
+			// If transition is valid and not waiting on other conditions
+			// and we're not doing a transition to self, unless the self-transition can inertialize
+			if (PotentialTransition.TargetState != CurrentState || ReferenceTransition.bAllowInertializationForSelfTransitions)
+			{
+				return true;
+			}
+
+			return false;
+		};
+
+		if (CanTransition())
 		{
 			bFoundValidTransition = true;
 
@@ -747,6 +778,8 @@ bool FAnimNode_StateMachine::FindValidTransition(const FAnimationUpdateContext&
 		}
 	}
 
+	const double WorldTime = Context.AnimInstanceProxy->GetAnimInstanceObject()->GetWorld()->GetTimeSeconds();
+
 	const int32 NumTransitions = StateInfo.Transitions.Num();
 	for (int32 TransitionIndex = 0; TransitionIndex < NumTransitions; ++TransitionIndex)
 	{
@@ -756,6 +789,16 @@ bool FAnimNode_StateMachine::FindValidTransition(const FAnimationUpdateContext&
 			continue;
 		}
 
+		const FAnimationTransitionBetweenStates& TransitionInfo = GetTransitionInfo(TransitionRule.TransitionIndex);
+		const int32 NextState = TransitionInfo.NextState;
+		const FBakedAnimationState& NextStateInfo = GetStateInfo(NextState);
+
+		// Check transition de-bounce time
+		if (TransitionInfo.MinTimeBeforeReentry >= 0.0f && (WorldTime - LastStateEntryTime[NextState]) < TransitionInfo.MinTimeBeforeReentry)
+		{
+			continue;
+		}
+
 		FAnimNode_TransitionResult* ResultNode = GetNodeFromPropertyIndex<FAnimNode_TransitionResult>(Context.AnimInstanceProxy->GetAnimInstanceObject(), AnimBlueprintClass, TransitionRule.CanTakeDelegateIndex);
 		float CrossfadeTimeAdjustment = 0.f;
 
@@ -798,8 +841,6 @@ bool FAnimNode_StateMachine::FindValidTransition(const FAnimationUpdateContext&
 						}
 					}
 
-					const FAnimationTransitionBetweenStates& TransitionInfo = GetTransitionInfo(TransitionRule.TransitionIndex);
-				
 					// For transitions that go to a conduit the user is not able to edit the transition's cross fade duration,
 					// therefore we force the cross fade duration to zero to ensure the transition is always triggerred upon reaching the end of the animation.
 					const float CrossfadeDuration = GetStateInfo(TransitionInfo.NextState).bIsAConduit ? 0.0f : TransitionInfo.CrossfadeDuration;
@@ -823,9 +864,6 @@ bool FAnimNode_StateMachine::FindValidTransition(const FAnimationUpdateContext&
 
 		if (ResultNode->bCanEnterTransition == TransitionRule.bDesiredTransitionReturnValue)
 		{
-			const int32 NextState = GetTransitionInfo(TransitionRule.TransitionIndex).NextState;
-			const FBakedAnimationState& NextStateInfo = GetStateInfo(NextState);
-
 			// if next state is a conduit we want to check for transitions using that state as the root
 			if (NextStateInfo.bIsAConduit)
 			{
@@ -1150,11 +1188,11 @@ void FAnimNode_StateMachine::SetStateInternal(int32 NewStateIndex)
 	ElapsedTime = 0.0f;
 }
 
-void FAnimNode_StateMachine::SetState(const FAnimationBaseContext& Context, int32 NewStateIndex)
+void FAnimNode_StateMachine::SetState(const FAnimationBaseContext& Context, int32 NewStateIndex, bool bAllowReEntry)
 {
 	SCOPE_CYCLE_COUNTER(Stat_StateMachineSetState);
 
-	if (NewStateIndex != CurrentState)
+	if (NewStateIndex != CurrentState || bAllowReEntry)
 	{
 		const int32 PrevStateIndex = CurrentState;
 		if(CurrentState != INDEX_NONE && CurrentState < OnGraphStatesExited.Num())
@@ -1180,6 +1218,9 @@ void FAnimNode_StateMachine::SetState(const FAnimationBaseContext& Context, int3
 
 		SetStateInternal(NewStateIndex);
 
+		// Record entry time for potential de-bouncing
+		LastStateEntryTime[CurrentState] = Context.AnimInstanceProxy->GetAnimInstanceObject()->GetWorld()->GetTimeSeconds();
+
 		// Clear any currently cached blend weights for asset player nodes.
 		// This stops any zero length blends holding on to old weights
 		for(int32 PlayerIndex : GetStateInfo(CurrentState).PlayerNodeIndices)
@@ -1307,7 +1348,22 @@ void FAnimNode_StateMachine::TransitionToState(const FAnimationUpdateContext& Co
 			NewTransition.StateBlendData[1].PerBoneBlendData.AddZeroed(NewTransition.BlendProfile->GetNumBlendEntries());
 		}
 
-		if (TransitionInfo.LogicType == ETransitionLogicType::TLT_Inertialization)
+		auto ShouldInertialize = [&TransitionInfo, &ExistingWeightOfNextState]()
+		{
+			// We are set to inertialize
+			if (TransitionInfo.LogicType == ETransitionLogicType::TLT_Inertialization)
+			{
+				return true;
+			}
+			// Fall back if we are allowed to and the next state has weight, to hide a pop if the state resets
+			else if (TransitionInfo.bAllowInertializationForSelfTransitions && ExistingWeightOfNextState > 0.0f)
+			{
+				return true;
+			}
+			return false;
+		};
+
+		if (ShouldInertialize())
 		{
 			UE::Anim::IInertializationRequester* InertializationRequester = Context.GetMessage<UE::Anim::IInertializationRequester>();
 			if (InertializationRequester)
@@ -1346,7 +1402,7 @@ void FAnimNode_StateMachine::TransitionToState(const FAnimationUpdateContext& Co
 	}
 
 	ConsumeMarkedTransitionEvents();
-	SetState(Context, NextState);
+	SetState(Context, NextState, TransitionInfo.bAllowInertializationForSelfTransitions);
 }
 
 float FAnimNode_StateMachine::GetStateWeight(int32 StateIndex) const
diff --git a/Engine/Source/Runtime/Engine/Public/Animation/AnimInstanceProxy.h b/Engine/Source/Runtime/Engine/Public/Animation/AnimInstanceProxy.h
--- a/Engine/Source/Runtime/Engine/Public/Animation/AnimInstanceProxy.h
+++ b/Engine/Source/Runtime/Engine/Public/Animation/AnimInstanceProxy.h
@@ -511,7 +511,7 @@ public:
 	ENGINE_API const FAnimNode_StateMachine* GetStateMachineInstanceFromName(FName MachineName) const;
 
 	/** Gets the runtime instance of the specified state machine */
-	ENGINE_API const FAnimNode_StateMachine* GetStateMachineInstance(int32 MachineIndex) const;
+	ENGINE_API const FAnimNode_StateMachine* GetStateMachineInstance(int32 NodeIndex) const;
 
 	/** Get the machine description for the specified instance. Does not rely on PRIVATE_MachineDescription being initialized */
 	static ENGINE_API const FBakedAnimationStateMachine* GetMachineDescription(IAnimClassInterface* AnimBlueprintClass, const FAnimNode_StateMachine* MachineInstance);
@@ -901,10 +901,10 @@ protected:
 	ENGINE_API void ClearAllTransitionEvents();
 
 	/** Returns whether or not the given event transition request has been queued */
-	ENGINE_API bool QueryTransitionEvent(int32 MachineIndex, int32 TransitionIndex, const FName& EventName) const;
+	ENGINE_API bool QueryTransitionEvent(int32 NodeIndex, int32 TransitionIndex, const FName& EventName) const;
 
 	/** Behaves like QueryTransitionEvent but additionally marks the event for consumption */
-	ENGINE_API bool QueryAndMarkTransitionEvent(int32 MachineIndex, int32 TransitionIndex, const FName& EventName);
+	ENGINE_API bool QueryAndMarkTransitionEvent(int32 NodeIndex, int32 TransitionIndex, const FName& EventName);
 
 	// Sets up a native transition delegate between states with PrevStateName and NextStateName, in the state machine with name MachineName.
 	// Note that a transition already has to exist for this to succeed
@@ -931,10 +931,10 @@ protected:
 	/** Gets the runtime instance desc of the state machine specified by name */
 	ENGINE_API const FBakedAnimationStateMachine* GetStateMachineInstanceDesc(FName MachineName) const;
 
-	/** Gets the index of the state machine matching MachineName */
+	/** Gets the node index of the state machine matching MachineName */
 	ENGINE_API int32 GetStateMachineIndex(FName MachineName) const;
 
-	/** Gets the index of the state machine */
+	/** Gets the node index of the state machine */
 	ENGINE_API int32 GetStateMachineIndex(FAnimNode_StateMachine* StateMachine) const;
 
 	ENGINE_API void GetStateMachineIndexAndDescription(FName InMachineName, int32& OutMachineIndex, const FBakedAnimationStateMachine** OutMachineDescription) const;
